ColorCorrectRegions: Merged from 5RES.

diff --git a/Engine/Plugins/Experimental/ColorCorrectRegions/Source/ColorCorrectRegions/Private/ColorCorrectRegion.cpp b/Engine/Plugins/Experimental/ColorCorrectRegions/Source/ColorCorrectRegions/Private/ColorCorrectRegion.cpp
--- a/Engine/Plugins/Experimental/ColorCorrectRegions/Source/ColorCorrectRegions/Private/ColorCorrectRegion.cpp
+++ b/Engine/Plugins/Experimental/ColorCorrectRegions/Source/ColorCorrectRegions/Private/ColorCorrectRegion.cpp
@@ -19,7 +19,7 @@ AColorCorrectRegion::AColorCorrectRegion(const FObjectInitializer& ObjectInitial
 	, ExcludeStencil(false)
 	, ColorCorrectRegionsSubsystem(nullptr)
 {
-
+	PrimaryActorTick.bCanEverTick = true;
 }
 
 void AColorCorrectRegion::BeginPlay()
@@ -57,6 +57,22 @@ void AColorCorrectRegion::BeginDestroy()
 	Super::BeginDestroy();
 }
 
+bool AColorCorrectRegion::ShouldTickIfViewportsOnly() const
+{
+	return true;
+}
+
+void AColorCorrectRegion::TickActor(float DeltaTime, ELevelTick TickType, FActorTickFunction& ThisTickFunction)
+{
+	Super::Tick(DeltaTime);
+	FTransform CurrentFrameTransform = GetTransform();
+	if (!PreviousFrameTransform.Equals(CurrentFrameTransform))
+	{
+		PreviousFrameTransform = CurrentFrameTransform;
+		GetActorBounds(true, BoxOrigin, BoxExtent);
+	}
+}
+
 void AColorCorrectRegion::Cleanup()
 {
 	ColorCorrectRegionsSubsystem = nullptr;
@@ -66,7 +82,6 @@ void AColorCorrectRegion::Cleanup()
 void AColorCorrectRegion::PostEditChangeProperty(struct FPropertyChangedEvent& PropertyChangedEvent)
 {
 	Super::PostEditChangeProperty(PropertyChangedEvent);
-
 	const FName PropertyName = PropertyChangedEvent.GetPropertyName();
 	if (PropertyName == GET_MEMBER_NAME_CHECKED(AColorCorrectRegion, Priority))
 	{
@@ -75,6 +90,7 @@ void AColorCorrectRegion::PostEditChangeProperty(struct FPropertyChangedEvent& P
 			ColorCorrectRegionsSubsystem->SortRegionsByPriority();
 		}
 	}
+	GetActorBounds(true, BoxOrigin, BoxExtent);
 }
 #endif //WITH_EDITOR
 
diff --git a/Engine/Plugins/Experimental/ColorCorrectRegions/Source/ColorCorrectRegions/Private/ColorCorrectRegionsSceneViewExtension.cpp b/Engine/Plugins/Experimental/ColorCorrectRegions/Source/ColorCorrectRegions/Private/ColorCorrectRegionsSceneViewExtension.cpp
--- a/Engine/Plugins/Experimental/ColorCorrectRegions/Source/ColorCorrectRegions/Private/ColorCorrectRegionsSceneViewExtension.cpp
+++ b/Engine/Plugins/Experimental/ColorCorrectRegions/Source/ColorCorrectRegions/Private/ColorCorrectRegionsSceneViewExtension.cpp
@@ -366,7 +366,7 @@ void FColorCorrectRegionsSceneViewExtension::PrePostProcessPass_RenderThread(FRD
 			}
 
 			FVector BoxCenter, BoxExtents;
-			Region->GetActorBounds(true, BoxCenter, BoxExtents);
+			Region->GetBounds(BoxCenter, BoxExtents);
 
 			// If bounding box is zero, then we don't need to do anything.
 			if (BoxExtents.IsNearlyZero())
diff --git a/Engine/Plugins/Experimental/ColorCorrectRegions/Source/ColorCorrectRegions/Public/ColorCorrectRegion.h b/Engine/Plugins/Experimental/ColorCorrectRegions/Source/ColorCorrectRegions/Public/ColorCorrectRegion.h
--- a/Engine/Plugins/Experimental/ColorCorrectRegions/Source/ColorCorrectRegions/Public/ColorCorrectRegion.h
+++ b/Engine/Plugins/Experimental/ColorCorrectRegions/Source/ColorCorrectRegions/Public/ColorCorrectRegion.h
@@ -95,6 +95,7 @@ public:
 #if WITH_EDITOR
 	/** Called when any of the properties are changed. */
 	virtual void PostEditChangeProperty(FPropertyChangedEvent& PropertyChangedEvent) override;
+
 #endif
 
 	/** To handle play in Editor, PIE and Standalone. These methods aggregate objects in play mode similarly to 
@@ -105,12 +106,29 @@ public:
 
 	virtual void BeginDestroy() override;
 
+	virtual void TickActor(float DeltaTime, enum ELevelTick TickType, FActorTickFunction& ThisTickFunction);
+	virtual bool ShouldTickIfViewportsOnly() const;
+
 	/** 
 	* We have to manage the lifetime of the region ourselves, because EndPlay is not guaranteed to be called 
 	* and BeginDestroy could be called from GC when it is too late.
 	*/
 	void Cleanup();
 
+	/**
+	* This is used on render thread, and not atomic on purpose to avoid stalling Render thread even for a little bit. 
+	*/
+	void GetBounds(FVector& InOutOrigin, FVector& InOutBoxExtent) const
+	{
+		InOutOrigin = BoxOrigin;
+		InOutBoxExtent = BoxExtent;
+	};
+
 private:
 	UColorCorrectRegionsSubsystem* ColorCorrectRegionsSubsystem;
+
+	FVector BoxOrigin;
+	FVector BoxExtent;
+
+	FTransform PreviousFrameTransform;
 };
