ColorCorrectRegions: Fixed a crash due to changes in View Frustrum.

diff --git a/Engine/Plugins/Experimental/ColorCorrectRegions/Source/ColorCorrectRegions/Private/ColorCorrectRegionsModule.cpp b/Engine/Plugins/Experimental/ColorCorrectRegions/Source/ColorCorrectRegions/Private/ColorCorrectRegionsModule.cpp
--- a/Engine/Plugins/Experimental/ColorCorrectRegions/Source/ColorCorrectRegions/Private/ColorCorrectRegionsModule.cpp
+++ b/Engine/Plugins/Experimental/ColorCorrectRegions/Source/ColorCorrectRegions/Private/ColorCorrectRegionsModule.cpp
@@ -22,4 +22,4 @@ void FColorCorrectRegionsModule::ShutdownModule()
 #undef LOCTEXT_NAMESPACE
 	
 IMPLEMENT_MODULE(FColorCorrectRegionsModule, ColorCorrectRegions);
-DEFINE_LOG_CATEGORY(ColorCorrectRegionsLogOutput);
+DEFINE_LOG_CATEGORY(ColorCorrectRegions);
diff --git a/Engine/Plugins/Experimental/ColorCorrectRegions/Source/ColorCorrectRegions/Private/ColorCorrectRegionsSceneViewExtension.cpp b/Engine/Plugins/Experimental/ColorCorrectRegions/Source/ColorCorrectRegions/Private/ColorCorrectRegionsSceneViewExtension.cpp
--- a/Engine/Plugins/Experimental/ColorCorrectRegions/Source/ColorCorrectRegions/Private/ColorCorrectRegionsSceneViewExtension.cpp
+++ b/Engine/Plugins/Experimental/ColorCorrectRegions/Source/ColorCorrectRegions/Private/ColorCorrectRegionsSceneViewExtension.cpp
@@ -4,6 +4,7 @@
 #include "RHI.h"
 #include "SceneView.h"
 #include "PostProcess/PostProcessing.h"
+#include "ColorCorrectRegionsModule.h"
 #include "ColorCorrectRegionsSubsystem.h"
 #include "ColorCorrectRegionsPostProcessMaterial.h"
 #include "ScreenPass.h"
@@ -169,7 +170,23 @@ namespace
 	{
 		TArray<FVector> Points;
 		Points.Reserve(6);
-		CalculatePlaneAABBIntersectionPoints(InView.ViewFrustum.Planes[4], InBoxCenter, InBoxExtents, Points);
+		static bool bNotifiedOfClippingPlaneError = false;
+
+		if (InView.bHasNearClippingPlane)
+		{
+			CalculatePlaneAABBIntersectionPoints(InView.NearClippingPlane, InBoxCenter, InBoxExtents, Points);
+		}
+		// Previously last plane was near clipping plane.
+		else if (InView.ViewFrustum.Planes.Num() == 5)
+		{
+			CalculatePlaneAABBIntersectionPoints(InView.ViewFrustum.Planes[4], InBoxCenter, InBoxExtents, Points);
+		}
+		else if (!bNotifiedOfClippingPlaneError)
+		{
+			bNotifiedOfClippingPlaneError = true;
+			UE_LOG(ColorCorrectRegions, Error, TEXT("Couldn't find a correct near clipping plane in View Frustrum"));
+		}
+
 		if (Points.Num() == 0)
 		{
 			return;
diff --git a/Engine/Plugins/Experimental/ColorCorrectRegions/Source/ColorCorrectRegions/Public/ColorCorrectRegionsModule.h b/Engine/Plugins/Experimental/ColorCorrectRegions/Source/ColorCorrectRegions/Public/ColorCorrectRegionsModule.h
--- a/Engine/Plugins/Experimental/ColorCorrectRegions/Source/ColorCorrectRegions/Public/ColorCorrectRegionsModule.h
+++ b/Engine/Plugins/Experimental/ColorCorrectRegions/Source/ColorCorrectRegions/Public/ColorCorrectRegionsModule.h
@@ -5,7 +5,7 @@
 #include "CoreMinimal.h"
 #include "Modules/ModuleManager.h"
 
-DECLARE_LOG_CATEGORY_EXTERN(ColorCorrectRegionsLogOutput, Log, All);
+DECLARE_LOG_CATEGORY_EXTERN(ColorCorrectRegions, Log, All);
 
 class FColorCorrectRegionsModule : public IModuleInterface
 {
