Color Correct Regions: Added a solution for checking Hidden status of Actor Components against the view. This makes Regions react to Display Cluster way of hiding objects from the view. Merged from Partner Branch 4.27.

diff --git a/Engine/Plugins/Experimental/ColorCorrectRegions/Source/ColorCorrectRegions/Private/ColorCorrectRegion.cpp b/Engine/Plugins/Experimental/ColorCorrectRegions/Source/ColorCorrectRegions/Private/ColorCorrectRegion.cpp
--- a/Engine/Plugins/Experimental/ColorCorrectRegions/Source/ColorCorrectRegions/Private/ColorCorrectRegion.cpp
+++ b/Engine/Plugins/Experimental/ColorCorrectRegions/Source/ColorCorrectRegions/Private/ColorCorrectRegion.cpp
@@ -71,6 +71,16 @@ void AColorCorrectRegion::TickActor(float DeltaTime, ELevelTick TickType, FActor
 		PreviousFrameTransform = CurrentFrameTransform;
 		GetActorBounds(true, BoxOrigin, BoxExtent);
 	}
+
+	// Display Cluster uses HiddenPrimitives to hide Primitive components from view. Store component id to be used
+	// on render thread.
+	if (const UPrimitiveComponent* FirstPrimitiveComponent = FindComponentByClass<UPrimitiveComponent>())
+	{
+		if (!(FirstPrimitiveId == FirstPrimitiveComponent->ComponentId))
+		{
+			FirstPrimitiveId = FirstPrimitiveComponent->ComponentId;
+		}
+	}
 }
 
 void AColorCorrectRegion::Cleanup()
diff --git a/Engine/Plugins/Experimental/ColorCorrectRegions/Source/ColorCorrectRegions/Private/ColorCorrectRegionsSceneViewExtension.cpp b/Engine/Plugins/Experimental/ColorCorrectRegions/Source/ColorCorrectRegions/Private/ColorCorrectRegionsSceneViewExtension.cpp
--- a/Engine/Plugins/Experimental/ColorCorrectRegions/Source/ColorCorrectRegions/Private/ColorCorrectRegionsSceneViewExtension.cpp
+++ b/Engine/Plugins/Experimental/ColorCorrectRegions/Source/ColorCorrectRegions/Private/ColorCorrectRegionsSceneViewExtension.cpp
@@ -377,7 +377,9 @@ void FColorCorrectRegionsSceneViewExtension::PrePostProcessPass_RenderThread(FRD
 #if WITH_EDITOR
 				Region->IsHiddenEd() ||
 #endif
-				Region->GetWorld() != ViewFamily.Scene->GetWorld())
+				Region->GetWorld() != ViewFamily.Scene->GetWorld() ||
+				// Display Cluster uses HiddenPrimitives to hide Primitive components from view.
+				View.HiddenPrimitives.Contains(Region->GetFirstPrimitiveId()))
 			{
 				continue;
 			}
diff --git a/Engine/Plugins/Experimental/ColorCorrectRegions/Source/ColorCorrectRegions/Public/ColorCorrectRegion.h b/Engine/Plugins/Experimental/ColorCorrectRegions/Source/ColorCorrectRegions/Public/ColorCorrectRegion.h
--- a/Engine/Plugins/Experimental/ColorCorrectRegions/Source/ColorCorrectRegions/Public/ColorCorrectRegion.h
+++ b/Engine/Plugins/Experimental/ColorCorrectRegions/Source/ColorCorrectRegions/Public/ColorCorrectRegion.h
@@ -124,6 +124,14 @@ public:
 		InOutBoxExtent = BoxExtent;
 	};
 
+	/**
+	* This method is used on render thread, and not atomic on purpose to avoid stalling Render thread even for a little bit.
+	*/
+	const FPrimitiveComponentId& GetFirstPrimitiveId() const
+	{
+		return FirstPrimitiveId;
+	};
+
 private:
 	UColorCorrectRegionsSubsystem* ColorCorrectRegionsSubsystem;
 
@@ -131,4 +139,5 @@ private:
 	FVector BoxExtent;
 
 	FTransform PreviousFrameTransform;
+	FPrimitiveComponentId FirstPrimitiveId;
 };
