Color Correct Regions: Raw actor class no longer appears in place actors menu.

diff --git a/Engine/Plugins/Experimental/ColorCorrectRegions/Source/ColorCorrectRegions/Public/ColorCorrectRegion.h b/Engine/Plugins/Experimental/ColorCorrectRegions/Source/ColorCorrectRegions/Public/ColorCorrectRegion.h
--- a/Engine/Plugins/Experimental/ColorCorrectRegions/Source/ColorCorrectRegions/Public/ColorCorrectRegion.h
+++ b/Engine/Plugins/Experimental/ColorCorrectRegions/Source/ColorCorrectRegions/Public/ColorCorrectRegion.h
@@ -39,7 +39,7 @@ enum class EColorCorrectRegionTemperatureType : uint8
  * or during BeginPlay/EndPlay to register itself. 
  * More information in ColorCorrectRegionsSubsytem.h
  */
-UCLASS(Blueprintable)
+UCLASS(Blueprintable, Abstract)
 class COLORCORRECTREGIONS_API AColorCorrectRegion : public AActor
 {
 	GENERATED_UCLASS_BODY()
