ColorCorrectRegions: Added billboard.

Used a slightly modified version of the place actor icon for now.

Also updated a couple of raw pointers to TObjectPtr in nDisplay Light Card Actor.

diff --git a/Engine/Plugins/Experimental/ColorCorrectRegions/Source/ColorCorrectRegions/Private/ColorCorrectRegion.cpp b/Engine/Plugins/Experimental/ColorCorrectRegions/Source/ColorCorrectRegions/Private/ColorCorrectRegion.cpp
--- a/Engine/Plugins/Experimental/ColorCorrectRegions/Source/ColorCorrectRegions/Private/ColorCorrectRegion.cpp
+++ b/Engine/Plugins/Experimental/ColorCorrectRegions/Source/ColorCorrectRegions/Private/ColorCorrectRegion.cpp
@@ -2,11 +2,14 @@
 
 #include "ColorCorrectRegion.h"
 #include "ColorCorrectRegionsSubsystem.h"
-#include "Engine/Classes/Components/MeshComponent.h"
+#include "Components/BillboardComponent.h"
 #include "CoreMinimal.h"
+#include "Engine/Classes/Components/MeshComponent.h"
+#include "Engine/Texture2D.h"
 #include "UObject/ConstructorHelpers.h"
 
-AColorCorrectRegion::AColorCorrectRegion(const FObjectInitializer& ObjectInitializer) : Super(ObjectInitializer)
+AColorCorrectRegion::AColorCorrectRegion(const FObjectInitializer& ObjectInitializer) 
+	: Super(ObjectInitializer)
 	, Type(EColorCorrectRegionsType::Sphere)
 	, Priority(0)
 	, Intensity(1.0)
@@ -20,6 +23,52 @@ AColorCorrectRegion::AColorCorrectRegion(const FObjectInitializer& ObjectInitial
 	, ColorCorrectRegionsSubsystem(nullptr)
 {
 	PrimaryActorTick.bCanEverTick = true;
+
+	// Add a scene component as our root
+	RootComponent = ObjectInitializer.CreateDefaultSubobject<USceneComponent>(this, TEXT("Root"));
+	RootComponent->SetMobility(EComponentMobility::Movable);
+
+#if WITH_EDITOR
+
+	// Create billboard component
+
+	if (GIsEditor && !IsRunningCommandlet())
+	{
+		// Structure to hold one-time initialization
+
+		struct FConstructorStatics
+		{
+			ConstructorHelpers::FObjectFinderOptional<UTexture2D> SpriteTextureObject;
+			FName ID_ColorCorrectRegion;
+			FText NAME_ColorCorrectRegion;
+
+			FConstructorStatics()
+				: SpriteTextureObject(TEXT("/ColorCorrectRegions/Icons/S_ColorCorrectRegionIcon"))
+				, ID_ColorCorrectRegion(TEXT("Color Correct Region"))
+				, NAME_ColorCorrectRegion(NSLOCTEXT("SpriteCategory", "ColorCorrectRegion", "Color Correct Region"))
+			{
+			}
+		};
+
+		static FConstructorStatics ConstructorStatics;
+
+		SpriteComponent = ObjectInitializer.CreateEditorOnlyDefaultSubobject<UBillboardComponent>(this, TEXT("Color Correct Region Icon"));
+		
+		if (SpriteComponent)
+		{
+			SpriteComponent->Sprite = ConstructorStatics.SpriteTextureObject.Get();
+			SpriteComponent->SpriteInfo.Category = ConstructorStatics.ID_ColorCorrectRegion;
+			SpriteComponent->SpriteInfo.DisplayName = ConstructorStatics.NAME_ColorCorrectRegion;
+			SpriteComponent->SetIsVisualizationComponent(true);
+			SpriteComponent->SetRelativeLocationAndRotation(FVector::ZeroVector, FRotator::ZeroRotator);
+			SpriteComponent->SetMobility(EComponentMobility::Movable);
+			SpriteComponent->bHiddenInGame = true;
+			
+			SpriteComponent->AttachToComponent(RootComponent, FAttachmentTransformRules::KeepRelativeTransform);
+		}
+	}
+
+#endif // WITH_EDITOR
 }
 
 void AColorCorrectRegion::BeginPlay()
diff --git a/Engine/Plugins/Experimental/ColorCorrectRegions/Source/ColorCorrectRegions/Public/ColorCorrectRegion.h b/Engine/Plugins/Experimental/ColorCorrectRegions/Source/ColorCorrectRegions/Public/ColorCorrectRegion.h
--- a/Engine/Plugins/Experimental/ColorCorrectRegions/Source/ColorCorrectRegions/Public/ColorCorrectRegion.h
+++ b/Engine/Plugins/Experimental/ColorCorrectRegions/Source/ColorCorrectRegions/Public/ColorCorrectRegion.h
@@ -11,6 +11,7 @@
 
 
 class UColorCorrectRegionsSubsystem;
+class UBillboardComponent;
 
 UENUM(BlueprintType)
 enum class EColorCorrectRegionsType : uint8 
@@ -93,6 +94,14 @@ public:
 	UPROPERTY(BlueprintReadWrite, EditAnywhere, Category = "Color Correction")
 	bool ExcludeStencil;
 
+#if WITH_EDITORONLY_DATA
+
+	/** Billboard component for this actor. */
+	UPROPERTY(Transient)
+	TObjectPtr<UBillboardComponent> SpriteComponent;
+
+#endif // WITH_EDITORONLY_DATA
+
 #if WITH_EDITOR
 	/** Called when any of the properties are changed. */
 	virtual void PostEditChangeProperty(FPropertyChangedEvent& PropertyChangedEvent) override;
diff --git a/Engine/Plugins/Runtime/nDisplay/Source/DisplayCluster/Public/DisplayClusterLightCardActor.h b/Engine/Plugins/Runtime/nDisplay/Source/DisplayCluster/Public/DisplayClusterLightCardActor.h
--- a/Engine/Plugins/Runtime/nDisplay/Source/DisplayCluster/Public/DisplayClusterLightCardActor.h
+++ b/Engine/Plugins/Runtime/nDisplay/Source/DisplayCluster/Public/DisplayClusterLightCardActor.h
@@ -136,7 +136,7 @@ public:
 	EDisplayClusterLightCardMask Mask;
 
 	UPROPERTY(EditAnywhere, BlueprintReadWrite, Category = "Appearance")
-	UTexture* Texture;
+	TObjectPtr<UTexture> Texture;
 
 	/** Light card color, before any modifier is applied */
 	UPROPERTY(EditAnywhere, BlueprintReadWrite, Category = "Appearance")
@@ -182,7 +182,7 @@ public:
 
 	/** Used to flag this light card as a proxy of a "real" light card. Used by the LightCard Editor */
 	UPROPERTY(Transient)
-	UTexture* PolygonMask = nullptr;
+	TObjectPtr<UTexture> PolygonMask = nullptr;
 
 protected:
 	UPROPERTY(EditDefaultsOnly, BlueprintReadOnly, Category = "Default")
