ColorCorrectRegions: Fixed always being hidden in game due to using the flag of the primitive id of a the new billboard instead of the mesh's.

diff --git a/Engine/Plugins/Experimental/ColorCorrectRegions/Source/ColorCorrectRegions/Private/ColorCorrectRegion.cpp b/Engine/Plugins/Experimental/ColorCorrectRegions/Source/ColorCorrectRegions/Private/ColorCorrectRegion.cpp
--- a/Engine/Plugins/Experimental/ColorCorrectRegions/Source/ColorCorrectRegions/Private/ColorCorrectRegion.cpp
+++ b/Engine/Plugins/Experimental/ColorCorrectRegions/Source/ColorCorrectRegions/Private/ColorCorrectRegion.cpp
@@ -3,6 +3,7 @@
 #include "ColorCorrectRegion.h"
 #include "ColorCorrectRegionsSubsystem.h"
 #include "Components/BillboardComponent.h"
+#include "Components/StaticMeshComponent.h"
 #include "CoreMinimal.h"
 #include "Engine/Classes/Components/MeshComponent.h"
 #include "Engine/Texture2D.h"
@@ -35,7 +36,7 @@ AColorCorrectRegion::AColorCorrectRegion(const FObjectInitializer& ObjectInitial
 	if (GIsEditor && !IsRunningCommandlet())
 	{
 		// Structure to hold one-time initialization
-
+	
 		struct FConstructorStatics
 		{
 			ConstructorHelpers::FObjectFinderOptional<UTexture2D> SpriteTextureObject;
@@ -53,7 +54,7 @@ AColorCorrectRegion::AColorCorrectRegion(const FObjectInitializer& ObjectInitial
 		static FConstructorStatics ConstructorStatics;
 
 		SpriteComponent = ObjectInitializer.CreateEditorOnlyDefaultSubobject<UBillboardComponent>(this, TEXT("Color Correct Region Icon"));
-		
+
 		if (SpriteComponent)
 		{
 			SpriteComponent->Sprite = ConstructorStatics.SpriteTextureObject.Get();
@@ -121,13 +122,13 @@ void AColorCorrectRegion::TickActor(float DeltaTime, ELevelTick TickType, FActor
 		GetActorBounds(true, BoxOrigin, BoxExtent);
 	}
 
-	// Display Cluster uses HiddenPrimitives to hide Primitive components from view. Store component id to be used
-	// on render thread.
-	if (const UPrimitiveComponent* FirstPrimitiveComponent = FindComponentByClass<UPrimitiveComponent>())
+	// Display Cluster uses HiddenPrimitives to hide Primitive components from view. 
+	// Store component id to be used on render thread.
+	if (const UStaticMeshComponent* FirstMeshComponent = FindComponentByClass<UStaticMeshComponent>())
 	{
-		if (!(FirstPrimitiveId == FirstPrimitiveComponent->ComponentId))
+		if (!(FirstPrimitiveId == FirstMeshComponent->ComponentId))
 		{
-			FirstPrimitiveId = FirstPrimitiveComponent->ComponentId;
+			FirstPrimitiveId = FirstMeshComponent->ComponentId;
 		}
 	}
 }
