Color Correct Regions: Added a new type of Color Correct Regions (Windowed CCR) which applies color grading to anything that is behind it.

diff --git a/Engine/Plugins/Experimental/ColorCorrectRegions/Shaders/Private/ColorCorrectRegionsCommon.usf b/Engine/Plugins/Experimental/ColorCorrectRegions/Shaders/Private/ColorCorrectRegionsCommon.usf
new file mode 100644
--- /dev/null
+++ b/Engine/Plugins/Experimental/ColorCorrectRegions/Shaders/Private/ColorCorrectRegionsCommon.usf
@@ -0,0 +1,417 @@
+// Copyright Epic Games, Inc. All Rights Reserved.
+
+#include "/Engine/Private/DeferredShadingCommon.ush"        // GetGBufferData()
+
+
+#ifdef NEEDS_SCENE_TEXTURES
+#undef NEEDS_SCENE_TEXTURES
+#endif
+#define NEEDS_SCENE_TEXTURES 1
+
+SCREEN_PASS_TEXTURE_VIEWPORT(PostProcessInput_0)
+SCREEN_PASS_TEXTURE_VIEWPORT(PostProcessInput_1)
+SCREEN_PASS_TEXTURE_VIEWPORT(PostProcessInput_2)
+SCREEN_PASS_TEXTURE_VIEWPORT(PostProcessInput_3)
+SCREEN_PASS_TEXTURE_VIEWPORT(PostProcessInput_4)
+SCREEN_PASS_TEXTURE_VIEWPORT(PostProcessOutput)
+
+Texture2D PostProcessInput_0_Texture;
+Texture2D PostProcessInput_1_Texture;
+Texture2D PostProcessInput_2_Texture;
+Texture2D PostProcessInput_3_Texture;
+Texture2D PostProcessInput_4_Texture;
+
+SamplerState PostProcessInput_0_Sampler;
+SamplerState PostProcessInput_1_Sampler;
+SamplerState PostProcessInput_2_Sampler;
+SamplerState PostProcessInput_3_Sampler;
+SamplerState PostProcessInput_4_Sampler;
+
+
+SamplerState PostProcessInputSampler;
+SamplerState PostProcessInput_BilinearSampler;
+
+#define PostProcessInput_0_SharedSampler PostProcessInput_0_Sampler
+#define PostProcessInput_1_SharedSampler PostProcessInput_1_Sampler
+#define PostProcessInput_2_SharedSampler PostProcessInput_2_Sampler
+#define PostProcessInput_3_SharedSampler PostProcessInput_3_Sampler
+#define PostProcessInput_4_SharedSampler PostProcessInput_4_Sampler
+
+//////////////////////////////////////////////////////////////////////////
+//! Must match ESceneTextureId
+
+#define PPI_SceneColor 0
+#define PPI_SceneDepth 1
+#define PPI_DiffuseColor 2
+#define PPI_SpecularColor 3
+#define PPI_SubsurfaceColor 4
+#define PPI_BaseColor 5
+#define PPI_Specular 6
+#define PPI_Metallic 7
+#define PPI_WorldNormal 8
+#define PPI_SeparateTranslucency 9
+#define PPI_Opacity 10
+#define PPI_Roughness 11
+#define PPI_MaterialAO 12
+#define PPI_CustomDepth 13
+#define PPI_PostProcessInput0 14
+#define PPI_PostProcessInput1 15
+#define PPI_PostProcessInput2 16
+#define PPI_PostProcessInput3 17
+#define PPI_PostProcessInput4 18
+#define PPI_PostProcessInput5 19 // (UNUSED)
+#define PPI_PostProcessInput6 20 // (UNUSED)
+#define PPI_DecalMask 21
+#define PPI_ShadingModelColor 22
+#define PPI_ShadingModelID 23
+#define PPI_AmbientOcclusion 24
+#define PPI_CustomStencil 25
+#define PPI_StoredBaseColor 26
+#define PPI_StoredSpecular 27
+#define PPI_Velocity 28
+#define PPI_WorldTangent 29
+#define PPI_Anisotropy 30
+
+//////////////////////////////////////////////////////////////////////////
+
+
+
+/**
+ * Parameters needed by pixel shader material inputs, related to Geometry.
+ * These are independent of vertex factory.
+ */
+struct FMaterialPixelParameters
+{
+
+	/** Interpolated vertex color, in linear color space. */
+	half4 VertexColor;
+
+	/** Normalized world space normal. */
+	half3 WorldNormal;
+
+	/** Normalized world space tangent. */
+	half3 WorldTangent;
+
+	/** Normalized world space reflected camera vector. */
+	half3 ReflectionVector;
+
+	/** Normalized world space camera vector, which is the vector from the point being shaded to the camera position. */
+	half3 CameraVector;
+
+	/** World space light vector, only valid when rendering a light function. */
+	half3 LightVector;
+
+	/**
+	 * Like SV_Position (.xy is pixel position at pixel center, z:DeviceZ, .w:SceneDepth)
+	 * using shader generated value SV_POSITION
+	 * Note: this is not relative to the current viewport.  RelativePixelPosition = MaterialParameters.SvPosition.xy - View.ViewRectMin.xy;
+	 */
+	MaterialFloat4 SvPosition;
+
+	/** Post projection position reconstructed from SvPosition, before the divide by W. left..top -1..1, bottom..top -1..1  within the viewport, W is the SceneDepth */
+	MaterialFloat4 ScreenPosition;
+
+	half UnMirrored;
+
+	half TwoSidedSign;
+
+	/**
+	 * Orthonormal rotation-only transform from tangent space to world space
+	 * The transpose(TangentToWorld) is WorldToTangent, and TangentToWorld[2] is WorldVertexNormal
+	 */
+	half3x3 TangentToWorld;
+
+	/**
+	 * Interpolated worldspace position of this pixel
+	 * todo: Make this TranslatedWorldPosition and also rename the VS WorldPosition to be TranslatedWorldPosition
+	 */
+	MaterialFloat3 AbsoluteWorldPosition;
+
+	/**
+	 * Interpolated worldspace position of this pixel, centered around the camera
+	 */
+	MaterialFloat3 WorldPosition_CamRelative;
+
+	/**
+	 * Interpolated worldspace position of this pixel, not including any world position offset or displacement.
+	 * Only valid if shader is compiled with NEEDS_WORLD_POSITION_EXCLUDING_SHADER_OFFSETS, otherwise just contains 0
+	 */
+	MaterialFloat3 WorldPosition_NoOffsets;
+
+	/**
+	 * Interpolated worldspace position of this pixel, not including any world position offset or displacement.
+	 * Only valid if shader is compiled with NEEDS_WORLD_POSITION_EXCLUDING_SHADER_OFFSETS, otherwise just contains 0
+	 */
+	MaterialFloat3 WorldPosition_NoOffsets_CamRelative;
+
+	/** Offset applied to the lighting position for translucency, used to break up aliasing artifacts. */
+	half3 LightingPositionOffset;
+
+	MaterialFloat AOMaterialMask;
+
+	// Index into View.PrimitiveSceneData
+	uint PrimitiveId;
+
+
+};
+
+
+/**
+ * Parameters calculated from the pixel material inputs.
+ */
+struct FPixelMaterialInputs
+{
+	MaterialFloat3 EmissiveColor;
+	MaterialFloat Opacity;
+	MaterialFloat OpacityMask;
+	MaterialFloat3 BaseColor;
+	MaterialFloat Metallic;
+	MaterialFloat Specular;
+	MaterialFloat Roughness;
+	MaterialFloat Anisotropy;
+	MaterialFloat3 Normal;
+	MaterialFloat3 Tangent;
+	MaterialFloat4 Subsurface;
+	MaterialFloat AmbientOcclusion;
+	MaterialFloat2 Refraction;
+	MaterialFloat PixelDepthOffset;
+	uint ShadingModel;
+
+};
+
+
+/**
+ * Parameters needed by vertex shader material inputs.
+ * These are independent of vertex factory.
+ */
+struct FMaterialVertexParameters
+{
+	// Position in the translated world (VertexFactoryGetWorldPosition).
+	// Previous position in the translated world (VertexFactoryGetPreviousWorldPosition) if
+	//    computing material's output for previous frame (See {BasePassVertex,Velocity}Shader.usf).
+	MaterialFloat3 WorldPosition;
+	// TangentToWorld[2] is WorldVertexNormal
+	half3x3 TangentToWorld;
+
+	// If either USE_INSTANCING or (IS_MESHPARTICLE_FACTORY && FEATURE_LEVEL >= FEATURE_LEVEL_SM4)
+	// is true, PrevFrameLocalToWorld is a per-instance transform
+	MaterialFloat4x4 PrevFrameLocalToWorld;
+
+	MaterialFloat3 PreSkinnedPosition;
+	MaterialFloat3 PreSkinnedNormal;
+
+	half4 VertexColor;
+
+	// Index into View.PrimitiveSceneData
+	uint PrimitiveId;
+
+};
+
+// Uniform material expressions.
+MaterialFloat3 sRGBtoAP1(MaterialFloat4 InputColor)
+{
+	const MaterialFloat3x3 sRGB_2_AP1 = mul(XYZ_2_AP1_MAT, mul(D65_2_D60_CAT, sRGB_2_XYZ_MAT));
+	MaterialFloat3 ColorAP1 = mul(sRGB_2_AP1, InputColor.xyz);
+	return ColorAP1;
+}
+
+MaterialFloat3 AP1toSRGB(MaterialFloat3 InputColor)
+{
+	const MaterialFloat3x3 AP1_2_sRGB = mul(XYZ_2_sRGB_MAT, mul(D60_2_D65_CAT, AP1_2_XYZ_MAT));
+	MaterialFloat3 ColorRGB = mul(AP1_2_sRGB, InputColor);
+	return ColorRGB;
+}
+
+MaterialFloat3 ComposureColorGradePerRange(MaterialFloat3 WorkingColor, MaterialFloat4 ColorSaturation, MaterialFloat4 ColorContrast, MaterialFloat4 ColorGamma, MaterialFloat4 ColorGain, MaterialFloat4 ColorOffset)
+{
+	MaterialFloat Luma = dot(WorkingColor, MaterialFloat3(0.2722287168, 0.6740817658, 0.0536895174));
+	WorkingColor = max(0, lerp(Luma.xxx, WorkingColor, ColorSaturation.xyz * ColorSaturation.w));
+	WorkingColor = pow(WorkingColor * (5.55555555), ColorContrast.xyz * ColorContrast.w) * 0.18;
+	WorkingColor = pow(WorkingColor, 1.0 / (ColorGamma.xyz * ColorGamma.w));
+	WorkingColor = WorkingColor * (ColorGain.xyz * ColorGain.w) + (ColorOffset.xyz + ColorOffset.w);
+	return WorkingColor;
+}
+
+
+/** Calculate a reflection vector about the specified world space normal. Optionally normalize this normal **/
+MaterialFloat3 ReflectionAboutCustomWorldNormal(FMaterialPixelParameters Parameters, MaterialFloat3 WorldNormal, bool bNormalizeInputNormal)
+{
+	if (bNormalizeInputNormal)
+	{
+		WorldNormal = normalize(WorldNormal);
+	}
+
+	return -Parameters.CameraVector + WorldNormal * dot(WorldNormal, Parameters.CameraVector) * 2.0;
+}
+
+
+MaterialFloat3 GetWorldPosition(FMaterialPixelParameters Parameters)
+{
+	return Parameters.AbsoluteWorldPosition;
+}
+
+MaterialFloat2 GetPixelPosition(FMaterialPixelParameters Parameters)
+{
+	return Parameters.SvPosition.xy - MaterialFloat2(PostProcessOutput_ViewportMin);
+}
+
+MaterialFloat2 GetViewportUV(FMaterialPixelParameters Parameters)
+{
+	return GetPixelPosition(Parameters) * PostProcessOutput_ViewportSizeInverse;
+}
+
+half3 GetMaterialNormalRaw(FPixelMaterialInputs PixelMaterialInputs)
+{
+	return PixelMaterialInputs.Normal;
+}
+
+half3 GetMaterialNormal(FMaterialPixelParameters Parameters, FPixelMaterialInputs PixelMaterialInputs)
+{
+	half3 RetNormal;
+	RetNormal = GetMaterialNormalRaw(PixelMaterialInputs);
+	return RetNormal;
+}
+
+// Returns whether a scene texture id is a for a post process input or not.
+bool IsPostProcessInputSceneTexture(const uint SceneTextureId)
+{
+	return (SceneTextureId >= PPI_PostProcessInput0 && SceneTextureId <= PPI_PostProcessInput6);
+}
+
+// Transforms viewport UV to scene texture's UV.
+MaterialFloat2 ViewportUVToSceneTextureUV(MaterialFloat2 ViewportUV, const uint SceneTextureId)
+{
+	if (IsPostProcessInputSceneTexture(SceneTextureId))
+	{
+		switch (SceneTextureId)
+		{
+		case PPI_PostProcessInput0:
+			return ViewportUV * PostProcessInput_0_UVViewportSize + PostProcessInput_0_UVViewportMin;
+		case PPI_PostProcessInput1:
+			return ViewportUV * PostProcessInput_1_UVViewportSize + PostProcessInput_1_UVViewportMin;
+		case PPI_PostProcessInput2:
+			return ViewportUV * PostProcessInput_2_UVViewportSize + PostProcessInput_2_UVViewportMin;
+		case PPI_PostProcessInput3:
+			return ViewportUV * PostProcessInput_3_UVViewportSize + PostProcessInput_3_UVViewportMin;
+		case PPI_PostProcessInput4:
+			return ViewportUV * PostProcessInput_4_UVViewportSize + PostProcessInput_4_UVViewportMin;
+		default:
+			return ViewportUV;
+		}
+	}
+
+	return ViewportUVToBufferUV(ViewportUV);
+}
+
+// Get default scene texture's UV.
+MaterialFloat2 GetDefaultSceneTextureUV(FMaterialPixelParameters Parameters, const uint SceneTextureId)
+{
+	return ViewportUVToSceneTextureUV(GetViewportUV(Parameters), SceneTextureId);
+}
+
+/** Rotates Position about the given axis by the given angle, in radians, and returns the offset to Position. */
+MaterialFloat3 RotateAboutAxis(MaterialFloat4 NormalizedRotationAxisAndAngle, MaterialFloat3 PositionOnAxis, MaterialFloat3 Position)
+{
+	// Project Position onto the rotation axis and find the closest point on the axis to Position
+	MaterialFloat3 ClosestPointOnAxis = PositionOnAxis + NormalizedRotationAxisAndAngle.xyz * dot(NormalizedRotationAxisAndAngle.xyz, Position - PositionOnAxis);
+	// Construct orthogonal axes in the plane of the rotation
+	MaterialFloat3 UAxis = Position - ClosestPointOnAxis;
+	MaterialFloat3 VAxis = cross(NormalizedRotationAxisAndAngle.xyz, UAxis);
+	MaterialFloat CosAngle;
+	MaterialFloat SinAngle;
+	sincos(NormalizedRotationAxisAndAngle.w, SinAngle, CosAngle);
+	// Rotate using the orthogonal axes
+	MaterialFloat3 R = UAxis * CosAngle + VAxis * SinAngle;
+	// Reconstruct the rotated world space position
+	MaterialFloat3 RotatedPosition = ClosestPointOnAxis + R;
+	// Convert from position to a position offset
+	return RotatedPosition - Position;
+}
+
+MaterialFloat3 RotateTranslateScale(MaterialFloat3 OriginalVec, MaterialFloat3 Translation, MaterialFloat3 Rotation, MaterialFloat3 Scale)
+{
+	// Translate
+	MaterialFloat3 Translated = (OriginalVec - Translation);
+	// Rotate
+	MaterialFloat3 RotatedZ = RotateAboutAxis(MaterialFloat4(MaterialFloat3(0.00000000, 0.00000000, 1.00000000), Rotation.b * -1), MaterialFloat3(0.00000000, 0.00000000, 0.00000000), Translated);
+	MaterialFloat3 RotatedTranslatedZ = (RotatedZ + Translated);
+	MaterialFloat3 RotatedY = RotateAboutAxis(MaterialFloat4(MaterialFloat3(0.00000000, 1.00000000, 0.00000000), Rotation.g), MaterialFloat3(0.00000000, 0.00000000, 0.00000000), RotatedTranslatedZ);
+	MaterialFloat3 RotatedTranslatedYZ = (RotatedY + RotatedTranslatedZ);
+	MaterialFloat3 RotatedX = RotateAboutAxis(MaterialFloat4(MaterialFloat3(1.00000000, 0.00000000, 0.00000000), Rotation.r), MaterialFloat3(0.00000000, 0.00000000, 0.00000000), RotatedTranslatedYZ);
+	MaterialFloat3 RotatedTranslated = (RotatedX + RotatedTranslatedYZ);
+	return (RotatedTranslated / (Scale));
+}
+
+
+/** Applies an offset to the scene texture lookup and decodes the HDR linear space color. */
+float4 SceneTextureLookup(float2 UV, int SceneTextureIndex, bool bFiltered)
+{
+	FScreenSpaceData ScreenSpaceData = GetScreenSpaceData(UV, false);
+	switch (SceneTextureIndex)
+	{
+		// order needs to match to ESceneTextureId
+
+	case PPI_SceneColor:
+		return float4(CalcSceneColor(UV), 0);
+	case PPI_SceneDepth:
+		return ScreenSpaceData.GBuffer.Depth;
+	case PPI_DiffuseColor:
+		return float4(ScreenSpaceData.GBuffer.DiffuseColor, 0);
+	case PPI_SpecularColor:
+		return float4(ScreenSpaceData.GBuffer.SpecularColor, 0);
+	case PPI_SubsurfaceColor:
+		return IsSubsurfaceModel(ScreenSpaceData.GBuffer.ShadingModelID) ? float4(ExtractSubsurfaceColor(ScreenSpaceData.GBuffer), ScreenSpaceData.GBuffer.CustomData.a) : ScreenSpaceData.GBuffer.CustomData;
+	case PPI_BaseColor:
+		return float4(ScreenSpaceData.GBuffer.BaseColor, 0);
+	case PPI_Specular:
+		return ScreenSpaceData.GBuffer.Specular;
+	case PPI_Metallic:
+		return ScreenSpaceData.GBuffer.Metallic;
+	case PPI_WorldNormal:
+		return float4(ScreenSpaceData.GBuffer.WorldNormal, 0);
+	case PPI_SeparateTranslucency:
+		return float4(1, 1, 1, 1);
+	case PPI_Opacity:
+		return ScreenSpaceData.GBuffer.CustomData.a;
+	case PPI_Roughness:
+		return ScreenSpaceData.GBuffer.Roughness;
+	case PPI_MaterialAO:
+		return ScreenSpaceData.GBuffer.GBufferAO;
+	case PPI_CustomDepth:
+		return ScreenSpaceData.GBuffer.CustomDepth;
+
+	case PPI_PostProcessInput0:
+		return Texture2DSample(PostProcessInput_0_Texture, bFiltered ? PostProcessInput_BilinearSampler : PostProcessInput_0_SharedSampler, UV);
+	case PPI_PostProcessInput1:
+		return Texture2DSample(PostProcessInput_1_Texture, bFiltered ? PostProcessInput_BilinearSampler : PostProcessInput_1_SharedSampler, UV);
+	case PPI_PostProcessInput2:
+		return Texture2DSample(PostProcessInput_2_Texture, bFiltered ? PostProcessInput_BilinearSampler : PostProcessInput_2_SharedSampler, UV);
+	case PPI_PostProcessInput3:
+		return Texture2DSample(PostProcessInput_3_Texture, bFiltered ? PostProcessInput_BilinearSampler : PostProcessInput_3_SharedSampler, UV);
+	case PPI_PostProcessInput4:
+		return Texture2DSample(PostProcessInput_4_Texture, bFiltered ? PostProcessInput_BilinearSampler : PostProcessInput_4_SharedSampler, UV);
+
+	case PPI_DecalMask:
+		return 0;  // material compiler will return an error
+	case PPI_ShadingModelColor:
+		return float4(GetShadingModelColor(ScreenSpaceData.GBuffer.ShadingModelID), 1);
+	case PPI_ShadingModelID:
+		return float4(ScreenSpaceData.GBuffer.ShadingModelID, 0, 0, 0);
+	case PPI_AmbientOcclusion:
+		return ScreenSpaceData.AmbientOcclusion;
+	case PPI_CustomStencil:
+		return ScreenSpaceData.GBuffer.CustomStencil;
+	case PPI_StoredBaseColor:
+		return float4(ScreenSpaceData.GBuffer.StoredBaseColor, 0);
+	case PPI_StoredSpecular:
+		return float4(ScreenSpaceData.GBuffer.StoredSpecular.rrr, 0);
+
+	case PPI_WorldTangent:
+		return float4(ScreenSpaceData.GBuffer.WorldTangent, 0);
+	case PPI_Anisotropy:
+		return ScreenSpaceData.GBuffer.Anisotropy;
+	default:
+		return float4(0, 0, 0, 0);
+	}
+}
+
diff --git a/Engine/Plugins/Experimental/ColorCorrectRegions/Shaders/Private/ColorCorrectRegionsShader.usf b/Engine/Plugins/Experimental/ColorCorrectRegions/Shaders/Private/ColorCorrectRegionsShader.usf
--- a/Engine/Plugins/Experimental/ColorCorrectRegions/Shaders/Private/ColorCorrectRegionsShader.usf
+++ b/Engine/Plugins/Experimental/ColorCorrectRegions/Shaders/Private/ColorCorrectRegionsShader.usf
@@ -2,86 +2,17 @@
 
 #include "/Engine/Private/Common.ush"
 #include "/Engine/Private/ColorUtils.ush"
-#include "/Engine/Private/SceneTexturesCommon.ush"
-#include "/Engine/Private/Random.ush"
 #include "/Engine/Private/DistanceField/GlobalDistanceFieldShared.ush"
+#include "/Engine/Generated/GeneratedUniformBuffers.ush" 
+#include "/Engine/Private/Random.ush"
+#include "/Engine/Private/SceneTexturesCommon.ush"
 #include "/Engine/Private/SceneData.ush"
 #include "/Engine/Private/ScreenPass.ush"
-#include "/Engine/Generated/GeneratedUniformBuffers.ush" 
 #include "/Engine/Private/SceneTextureParameters.ush"
 
 
 #include "/ColorCorrectRegionsShaders/Private/ColorCorrectRegionsColorTemperature.ush"
-
-#ifdef NEEDS_SCENE_TEXTURES
-	#undef NEEDS_SCENE_TEXTURES
-#endif
-#define NEEDS_SCENE_TEXTURES 1
-
-SCREEN_PASS_TEXTURE_VIEWPORT(PostProcessInput_0)
-SCREEN_PASS_TEXTURE_VIEWPORT(PostProcessInput_1)
-SCREEN_PASS_TEXTURE_VIEWPORT(PostProcessInput_2)
-SCREEN_PASS_TEXTURE_VIEWPORT(PostProcessInput_3)
-SCREEN_PASS_TEXTURE_VIEWPORT(PostProcessInput_4)
-SCREEN_PASS_TEXTURE_VIEWPORT(PostProcessOutput)
-
-Texture2D PostProcessInput_0_Texture;
-Texture2D PostProcessInput_1_Texture;
-Texture2D PostProcessInput_2_Texture;
-Texture2D PostProcessInput_3_Texture;
-Texture2D PostProcessInput_4_Texture;
-
-SamplerState PostProcessInput_0_Sampler;
-SamplerState PostProcessInput_1_Sampler;
-SamplerState PostProcessInput_2_Sampler;
-SamplerState PostProcessInput_3_Sampler;
-SamplerState PostProcessInput_4_Sampler;
-
-
-SamplerState PostProcessInputSampler;
-SamplerState PostProcessInput_BilinearSampler;
-
-#define PostProcessInput_0_SharedSampler PostProcessInput_0_Sampler
-#define PostProcessInput_1_SharedSampler PostProcessInput_1_Sampler
-#define PostProcessInput_2_SharedSampler PostProcessInput_2_Sampler
-#define PostProcessInput_3_SharedSampler PostProcessInput_3_Sampler
-#define PostProcessInput_4_SharedSampler PostProcessInput_4_Sampler
-
-
-//////////////////////////////////////////////////////////////////////////
-//! Must match ESceneTextureId
-
-#define PPI_SceneColor 0
-#define PPI_SceneDepth 1
-#define PPI_DiffuseColor 2
-#define PPI_SpecularColor 3
-#define PPI_SubsurfaceColor 4
-#define PPI_BaseColor 5
-#define PPI_Specular 6
-#define PPI_Metallic 7
-#define PPI_WorldNormal 8
-#define PPI_SeparateTranslucency 9
-#define PPI_Opacity 10
-#define PPI_Roughness 11
-#define PPI_MaterialAO 12
-#define PPI_CustomDepth 13
-#define PPI_PostProcessInput0 14
-#define PPI_PostProcessInput1 15
-#define PPI_PostProcessInput2 16
-#define PPI_PostProcessInput3 17
-#define PPI_PostProcessInput4 18
-#define PPI_PostProcessInput5 19 // (UNUSED)
-#define PPI_PostProcessInput6 20 // (UNUSED)
-#define PPI_DecalMask 21
-#define PPI_ShadingModelColor 22
-#define PPI_ShadingModelID 23
-#define PPI_AmbientOcclusion 24
-#define PPI_CustomStencil 25
-#define PPI_StoredBaseColor 26
-#define PPI_StoredSpecular 27
-#define PPI_Velocity 28
-#define PPI_WorldTangent 29
-#define PPI_Anisotropy 30
+#include "/ColorCorrectRegionsShaders/Private/ColorCorrectRegionsCommon.usf"
 
 //////////////////////////////////////////////////////////////////////////
 
@@ -90,412 +21,22 @@ SamplerState PostProcessInput_BilinearSampler;
 #define SHAPE_TYPE_BOX 1
 #define SHAPE_TYPE_CYLINDER 2
 #define SHAPE_TYPE_CONE 3
+#define SHAPE_TYPE_PLANE 4
 
 #define TEMPERATURE_TYPE_LEGACY 0
 #define TEMPERATURE_TYPE_WHITE_BALANCE 1
 #define TEMPERATURE_TYPE_COLOR_TEMP 2
 #define TEMPERATURE_TYPE_SKIP 3
 
-
-/** 
- * Parameters calculated from the pixel material inputs.
- */
-struct FPixelMaterialInputs
-{
-	MaterialFloat3 EmissiveColor;
-	MaterialFloat Opacity;
-	MaterialFloat OpacityMask;
-	MaterialFloat3 BaseColor;
-	MaterialFloat Metallic;
-	MaterialFloat Specular;
-	MaterialFloat Roughness;
-	MaterialFloat Anisotropy;
-	MaterialFloat3 Normal;
-	MaterialFloat3 Tangent;
-	MaterialFloat4 Subsurface;
-	MaterialFloat AmbientOcclusion;
-	MaterialFloat2 Refraction;
-	MaterialFloat PixelDepthOffset;
-	uint ShadingModel;
-
-};
-
-
-/** 
- * Parameters needed by pixel shader material inputs, related to Geometry.
- * These are independent of vertex factory.
- */
-struct FMaterialPixelParameters
-{
-
-    /** Interpolated vertex color, in linear color space. */
-	half4 VertexColor;
-
-    /** Normalized world space normal. */
-	half3 WorldNormal;
-    
-    /** Normalized world space tangent. */
-	half3 WorldTangent;
-
-    /** Normalized world space reflected camera vector. */
-	half3 ReflectionVector;
-
-    /** Normalized world space camera vector, which is the vector from the point being shaded to the camera position. */
-	half3 CameraVector;
-
-    /** World space light vector, only valid when rendering a light function. */
-	half3 LightVector;
-
-    /**
-     * Like SV_Position (.xy is pixel position at pixel center, z:DeviceZ, .w:SceneDepth)
-     * using shader generated value SV_POSITION
-     * Note: this is not relative to the current viewport.  RelativePixelPosition = MaterialParameters.SvPosition.xy - View.ViewRectMin.xy;
-     */
-	float4 SvPosition;
-        
-    /** Post projection position reconstructed from SvPosition, before the divide by W. left..top -1..1, bottom..top -1..1  within the viewport, W is the SceneDepth */
-	float4 ScreenPosition;
-
-	half UnMirrored;
-
-	half TwoSidedSign;
-
-    /**
-     * Orthonormal rotation-only transform from tangent space to world space
-     * The transpose(TangentToWorld) is WorldToTangent, and TangentToWorld[2] is WorldVertexNormal
-     */
-	half3x3 TangentToWorld;
-
-    /** 
-     * Interpolated worldspace position of this pixel
-     * todo: Make this TranslatedWorldPosition and also rename the VS WorldPosition to be TranslatedWorldPosition
-     */
-	float3 AbsoluteWorldPosition;
-
-    /** 
-     * Interpolated worldspace position of this pixel, centered around the camera
-     */
-	float3 WorldPosition_CamRelative;
-
-    /** 
-     * Interpolated worldspace position of this pixel, not including any world position offset or displacement.
-     * Only valid if shader is compiled with NEEDS_WORLD_POSITION_EXCLUDING_SHADER_OFFSETS, otherwise just contains 0
-     */
-	float3 WorldPosition_NoOffsets;
-
-    /** 
-     * Interpolated worldspace position of this pixel, not including any world position offset or displacement.
-     * Only valid if shader is compiled with NEEDS_WORLD_POSITION_EXCLUDING_SHADER_OFFSETS, otherwise just contains 0
-     */
-	float3 WorldPosition_NoOffsets_CamRelative;
-
-    /** Offset applied to the lighting position for translucency, used to break up aliasing artifacts. */
-	half3 LightingPositionOffset;
-
-	float AOMaterialMask;
-
-    // Index into View.PrimitiveSceneData
-	uint PrimitiveId;
-
-
-};
-
-// @todo compat hack
-FMaterialPixelParameters MakeInitializedMaterialPixelParameters()
-{
-	FMaterialPixelParameters MPP;
-	MPP = (FMaterialPixelParameters) 0;
-	MPP.TangentToWorld = float3x3(1, 0, 0, 0, 1, 0, 0, 0, 1);
-	return MPP;
-}
-
-/** 
- * Parameters needed by vertex shader material inputs.
- * These are independent of vertex factory.
- */
-struct FMaterialVertexParameters
-{
-    // Position in the translated world (VertexFactoryGetWorldPosition).
-    // Previous position in the translated world (VertexFactoryGetPreviousWorldPosition) if
-    //    computing material's output for previous frame (See {BasePassVertex,Velocity}Shader.usf).
-	float3 WorldPosition;
-    // TangentToWorld[2] is WorldVertexNormal
-	half3x3 TangentToWorld;
-
-    // If either USE_INSTANCING or (IS_MESHPARTICLE_FACTORY && FEATURE_LEVEL >= FEATURE_LEVEL_SM4)
-    // is true, PrevFrameLocalToWorld is a per-instance transform
-	float4x4 PrevFrameLocalToWorld;
-
-	float3 PreSkinnedPosition;
-	float3 PreSkinnedNormal;
-
-	half4 VertexColor;
-
-    // Index into View.PrimitiveSceneData
-	uint PrimitiveId;
-
-};
-
-/**
- * Returns the upper 3x3 portion of the LocalToWorld matrix.
- */
-MaterialFloat3x3 GetLocalToWorld3x3(uint PrimitiveId)
-{
-	return LWCToFloat3x3(GetPrimitiveData(PrimitiveId).LocalToWorld);
-}
-
-MaterialFloat3x3 GetLocalToWorld3x3()
-{
-	return LWCToFloat3x3(GetPrimitiveDataFromUniformBuffer().LocalToWorld);
-}
-
-float3 GetTranslatedWorldPosition(FMaterialVertexParameters Parameters)
-{
-	return Parameters.WorldPosition;
-}
-
-float3 GetWorldPosition(FMaterialVertexParameters Parameters)
-{
-	return GetTranslatedWorldPosition(Parameters) - LWCHackToFloat(ResolvedView.PreViewTranslation);
-}
-
-float3 GetWorldPosition(FMaterialPixelParameters Parameters)
-{
-	return Parameters.AbsoluteWorldPosition;
-}
-
-float3 GetTranslatedWorldPosition(FMaterialPixelParameters Parameters)
-{
-	return Parameters.WorldPosition_CamRelative;
-}
-
-float4 GetScreenPosition(FMaterialVertexParameters Parameters)
-{
-	return mul(float4(Parameters.WorldPosition, 1.0f), ResolvedView.TranslatedWorldToClip);
-}
-
-float4 GetScreenPosition(FMaterialPixelParameters Parameters)
-{
-	return Parameters.ScreenPosition;
-}
-
-float2 GetSceneTextureUV(FMaterialVertexParameters Parameters)
-{
-	return ScreenAlignedPosition(GetScreenPosition(Parameters));
-}
-
-float2 GetSceneTextureUV(FMaterialPixelParameters Parameters)
-{
-	return SvPositionToBufferUV(Parameters.SvPosition);
-}
-
-float2 GetViewportUV(FMaterialVertexParameters Parameters)
-{
-    return Parameters.WorldPosition.xy;
-}
-
-float2 GetPixelPosition(FMaterialVertexParameters Parameters)
-{
-	return GetViewportUV(Parameters) * View.ViewSizeAndInvSize.xy;
-}
-
-float2 GetPixelPosition(FMaterialPixelParameters Parameters)
-{
-    return Parameters.SvPosition.xy - float2(PostProcessOutput_ViewportMin);
-}
-
-float2 GetViewportUV(FMaterialPixelParameters Parameters)
-{
-    return GetPixelPosition(Parameters) * PostProcessOutput_ViewportSizeInverse;
-}
-
-half3 GetMaterialNormalRaw(FPixelMaterialInputs PixelMaterialInputs)
-{
-	return PixelMaterialInputs.Normal;
-}
-
-half3 GetMaterialNormal(FMaterialPixelParameters Parameters, FPixelMaterialInputs PixelMaterialInputs)
-{
-	half3 RetNormal;
-	RetNormal = GetMaterialNormalRaw(PixelMaterialInputs);
-	return RetNormal;
-}
-
-// Returns whether a scene texture id is a for a post process input or not.
-bool IsPostProcessInputSceneTexture(const uint SceneTextureId)
-{
-	return (SceneTextureId >= PPI_PostProcessInput0 && SceneTextureId <= PPI_PostProcessInput6);
-}
-
-// Transforms viewport UV to scene texture's UV.
-MaterialFloat2 ViewportUVToSceneTextureUV(MaterialFloat2 ViewportUV, const uint SceneTextureId)
-{
-    if (IsPostProcessInputSceneTexture(SceneTextureId))
-    {
-        switch (SceneTextureId)
-        {
-        case PPI_PostProcessInput0:
-            return ViewportUV * PostProcessInput_0_UVViewportSize + PostProcessInput_0_UVViewportMin;
-        case PPI_PostProcessInput1:
-            return ViewportUV * PostProcessInput_1_UVViewportSize + PostProcessInput_1_UVViewportMin;
-        case PPI_PostProcessInput2:
-            return ViewportUV * PostProcessInput_2_UVViewportSize + PostProcessInput_2_UVViewportMin;
-        case PPI_PostProcessInput3:
-            return ViewportUV * PostProcessInput_3_UVViewportSize + PostProcessInput_3_UVViewportMin;
-        case PPI_PostProcessInput4:
-            return ViewportUV * PostProcessInput_4_UVViewportSize + PostProcessInput_4_UVViewportMin;
-        default:
-            return ViewportUV;
-        }
-    }
-
-	return ViewportUVToBufferUV(ViewportUV);
-}
-
-// Get default scene texture's UV.
-MaterialFloat2 GetDefaultSceneTextureUV(FMaterialVertexParameters Parameters, const uint SceneTextureId)
-{
-	return GetSceneTextureUV(Parameters);
-}
-
-// Get default scene texture's UV.
-MaterialFloat2 GetDefaultSceneTextureUV(FMaterialPixelParameters Parameters, const uint SceneTextureId)
-{
-	return ViewportUVToSceneTextureUV(GetViewportUV(Parameters), SceneTextureId);
-}
-
-/** Return the object's position in world space */
-float3 GetObjectWorldPosition(FMaterialPixelParameters Parameters)
-{
-    return LWCHackToFloat(GetPrimitiveData(Parameters.PrimitiveId).ObjectWorldPosition);
-}
-
-
-/** Return the object's position in world space. For instanced meshes, this returns the instance position. */
-float3 GetObjectWorldPosition(FMaterialVertexParameters Parameters)
-{
-    return LWCHackToFloat(GetPrimitiveData(Parameters.PrimitiveId).ObjectWorldPosition);
-}
-
-/** Rotates Position about the given axis by the given angle, in radians, and returns the offset to Position. */
-float3 RotateAboutAxis(float4 NormalizedRotationAxisAndAngle, float3 PositionOnAxis, float3 Position)
-{
-    // Project Position onto the rotation axis and find the closest point on the axis to Position
-	float3 ClosestPointOnAxis = PositionOnAxis + NormalizedRotationAxisAndAngle.xyz * dot(NormalizedRotationAxisAndAngle.xyz, Position - PositionOnAxis);
-    // Construct orthogonal axes in the plane of the rotation
-	float3 UAxis = Position - ClosestPointOnAxis;
-	float3 VAxis = cross(NormalizedRotationAxisAndAngle.xyz, UAxis);
-	float CosAngle;
-	float SinAngle;
-	sincos(NormalizedRotationAxisAndAngle.w, SinAngle, CosAngle);
-    // Rotate using the orthogonal axes
-	float3 R = UAxis * CosAngle + VAxis * SinAngle;
-    // Reconstruct the rotated world space position
-	float3 RotatedPosition = ClosestPointOnAxis + R;
-    // Convert from position to a position offset
-	return RotatedPosition - Position;
-}
-
-
-/** Calculate a reflection vector about the specified world space normal. Optionally normalize this normal **/
-MaterialFloat3 ReflectionAboutCustomWorldNormal(FMaterialPixelParameters Parameters, MaterialFloat3 WorldNormal, bool bNormalizeInputNormal)
-{
-	if (bNormalizeInputNormal)
-	{
-		WorldNormal = normalize(WorldNormal);
-	}
-
-	return -Parameters.CameraVector + WorldNormal * dot(WorldNormal, Parameters.CameraVector) * 2.0;
-}
-
-
-
-#if SHADING_PATH_DEFERRED
-
-#include "/Engine/Private/DeferredShadingCommon.ush"        // GetGBufferData()
-
-
-/** Applies an offset to the scene texture lookup and decodes the HDR linear space color. */
-float4 SceneTextureLookup(float2 UV, int SceneTextureIndex, bool bFiltered)
-{
-    FScreenSpaceData ScreenSpaceData = GetScreenSpaceData(UV, false);
-    switch(SceneTextureIndex)
-    {
-        // order needs to match to ESceneTextureId
-
-        case PPI_SceneColor:
-            return float4(CalcSceneColor(UV), 0);
-        case PPI_SceneDepth:
-            return ScreenSpaceData.GBuffer.Depth;
-        case PPI_DiffuseColor:
-            return float4(ScreenSpaceData.GBuffer.DiffuseColor, 0);
-        case PPI_SpecularColor:
-            return float4(ScreenSpaceData.GBuffer.SpecularColor, 0);
-        case PPI_SubsurfaceColor:
-            return IsSubsurfaceModel(ScreenSpaceData.GBuffer.ShadingModelID) ? float4( ExtractSubsurfaceColor(ScreenSpaceData.GBuffer), ScreenSpaceData.GBuffer.CustomData.a ) : ScreenSpaceData.GBuffer.CustomData;
-        case PPI_BaseColor:
-            return float4(ScreenSpaceData.GBuffer.BaseColor, 0);
-        case PPI_Specular:
-            return ScreenSpaceData.GBuffer.Specular;
-        case PPI_Metallic:
-            return ScreenSpaceData.GBuffer.Metallic;
-        case PPI_WorldNormal:
-            return float4(ScreenSpaceData.GBuffer.WorldNormal, 0);
-        case PPI_SeparateTranslucency:
-            return float4(1, 1, 1, 1);    // todo
-        case PPI_Opacity:
-            return ScreenSpaceData.GBuffer.CustomData.a;
-        case PPI_Roughness:
-            return ScreenSpaceData.GBuffer.Roughness;
-        case PPI_MaterialAO:
-            return ScreenSpaceData.GBuffer.GBufferAO;
-        case PPI_CustomDepth:
-            return ScreenSpaceData.GBuffer.CustomDepth;
-
-        case PPI_PostProcessInput0:
-            return Texture2DSample(PostProcessInput_0_Texture, bFiltered ? PostProcessInput_BilinearSampler : PostProcessInput_0_SharedSampler, UV);
-        case PPI_PostProcessInput1:
-            return Texture2DSample(PostProcessInput_1_Texture, bFiltered ? PostProcessInput_BilinearSampler : PostProcessInput_1_SharedSampler, UV);
-        case PPI_PostProcessInput2:
-            return Texture2DSample(PostProcessInput_2_Texture, bFiltered ? PostProcessInput_BilinearSampler : PostProcessInput_2_SharedSampler, UV);
-        case PPI_PostProcessInput3:
-            return Texture2DSample(PostProcessInput_3_Texture, bFiltered ? PostProcessInput_BilinearSampler : PostProcessInput_3_SharedSampler, UV);
-        case PPI_PostProcessInput4:
-            return Texture2DSample(PostProcessInput_4_Texture, bFiltered ? PostProcessInput_BilinearSampler : PostProcessInput_4_SharedSampler, UV);
-
-        case PPI_DecalMask:
-            return 0;  // material compiler will return an error
-        case PPI_ShadingModelColor:
-            return float4(GetShadingModelColor(ScreenSpaceData.GBuffer.ShadingModelID), 1);
-        case PPI_ShadingModelID:
-            return float4(ScreenSpaceData.GBuffer.ShadingModelID, 0, 0, 0);
-        case PPI_AmbientOcclusion:
-            return ScreenSpaceData.AmbientOcclusion;
-        case PPI_CustomStencil:
-            return ScreenSpaceData.GBuffer.CustomStencil;
-        case PPI_StoredBaseColor:
-            return float4(ScreenSpaceData.GBuffer.StoredBaseColor, 0);
-        case PPI_StoredSpecular:
-            return float4(ScreenSpaceData.GBuffer.StoredSpecular.rrr, 0);
-
-        case PPI_WorldTangent:
-            return float4(ScreenSpaceData.GBuffer.WorldTangent, 0);
-        case PPI_Anisotropy:
-            return ScreenSpaceData.GBuffer.Anisotropy;
-        default:
-            return float4(0, 0, 0, 0);
-    }
-}
-
-#endif // SHADING_PATH_DEFERRED
-
+#define CCR_SHADER_DISPLAY_BOUNDING_RECT 0
+#define CLIP_PIXELS_OUTSIDE_AABB 1
 
 struct FExternalExpressions
 {
 	float3 Rotate;
 	float3 Translate;
 	float3 Scale;
-	
+
 	float4 ColorSaturation;
 	float4 ColorContrast;
 	float4 ColorGamma;
@@ -509,13 +50,13 @@ struct FExternalExpressions
 	float4 ColorGainShadows;
 	float4 ColorOffsetShadows;
 	float ShadowMax;
-	
+
 	float4 ColorSaturationMidtones;
 	float4 ColorContrastMidtones;
 	float4 ColorGammaMidtones;
 	float4 ColorGainMidtones;
 	float4 ColorOffsetMidtones;
-	
+
 	float4 ColorSaturationHighlights;
 	float4 ColorContrastHighlights;
 	float4 ColorGammaHighlights;
@@ -527,139 +68,33 @@ struct FExternalExpressions
 	float WhiteTemp;
 	float Inner;
 	float Outer;
-	
+
 	float Falloff;
 	float Invert;
 	float Intensity;
-	
+
 	float FakeLight;
 	float ExcludeStencil;
 
 	float4 SceneTexture;
 };
 
-struct FOutputData
-{
-	float3 EmissiveColor;
-	float Opacity;
-};
-
-		
-// Uniform material expressions.
-MaterialFloat3 RGBtoAP1(MaterialFloat4 InputColor)
-{
-	const float3x3 sRGB_2_AP1 = mul(XYZ_2_AP1_MAT, mul(D65_2_D60_CAT, sRGB_2_XYZ_MAT));
-	const float3x3 AP1_2_sRGB = mul(XYZ_2_sRGB_MAT, mul(D60_2_D65_CAT, AP1_2_XYZ_MAT));
 
-	float3 ColorAP1 = mul(sRGB_2_AP1, InputColor.xyz);
-	float3 ColorRGB = mul(AP1_2_sRGB, InputColor.xyz);
-
-	return ColorAP1;
-}
-
-MaterialFloat3 AP1toRGB(MaterialFloat3 InputColor)
+MaterialFloat3 ProcessColor(in FExternalExpressions ExternalExpressions)
 {
-	const float3x3 sRGB_2_AP1 = mul(XYZ_2_AP1_MAT, mul(D65_2_D60_CAT, sRGB_2_XYZ_MAT));
-	const float3x3 AP1_2_sRGB = mul(XYZ_2_sRGB_MAT, mul(D60_2_D65_CAT, AP1_2_XYZ_MAT));
-
-	float3 ColorAP1 = mul(sRGB_2_AP1, InputColor);
-	float3 ColorRGB = mul(AP1_2_sRGB, InputColor);
-
-	return ColorRGB;
-}
-
-MaterialFloat3 ComposureColorGradePerRange(MaterialFloat3 WorkingColor, MaterialFloat4 ColorSaturation, MaterialFloat4 ColorContrast, MaterialFloat4 ColorGamma, MaterialFloat4 ColorGain, MaterialFloat4 ColorOffset)
-{
-	float Luma = dot(WorkingColor, float3(0.2722287168, 0.6740817658, 0.0536895174));
-	WorkingColor = max(0, lerp(Luma.xxx, WorkingColor, ColorSaturation.xyz * ColorSaturation.w));
-	WorkingColor = pow(WorkingColor * (5.55555555), ColorContrast.xyz * ColorContrast.w) * 0.18;
-	WorkingColor = pow(WorkingColor, 1.0 / (ColorGamma.xyz * ColorGamma.w));
-	WorkingColor = WorkingColor * (ColorGain.xyz * ColorGain.w) + (ColorOffset.xyz + ColorOffset.w);
-	return WorkingColor;
-}
-
-float4 ProcessRegion(in out FMaterialPixelParameters Parameters, in FExternalExpressions ExternalExpressions)
-{
-	// Translate
-	MaterialFloat3 Translated = (GetWorldPosition(Parameters) - ExternalExpressions.Translate.rgb);
-	// Rotate
-	MaterialFloat3 RotatedZ = RotateAboutAxis(MaterialFloat4(MaterialFloat3(0.00000000, 0.00000000, 1.00000000), ExternalExpressions.Rotate.b * -1), MaterialFloat3(0.00000000, 0.00000000, 0.00000000), Translated);
-	MaterialFloat3 RotatedTranslatedZ = (RotatedZ + Translated);
-	MaterialFloat3 RotatedY = RotateAboutAxis(MaterialFloat4(MaterialFloat3(0.00000000, 1.00000000, 0.00000000), ExternalExpressions.Rotate.g), MaterialFloat3(0.00000000, 0.00000000, 0.00000000), RotatedTranslatedZ);
-	MaterialFloat3 RotatedTranslatedYZ = (RotatedY + RotatedTranslatedZ);
-	MaterialFloat3 RotatedX = RotateAboutAxis(MaterialFloat4(MaterialFloat3(1.00000000, 0.00000000, 0.00000000), ExternalExpressions.Rotate.r), MaterialFloat3(0.00000000, 0.00000000, 0.00000000), RotatedTranslatedYZ);
-	MaterialFloat3 RotatedTranslated = (RotatedX + RotatedTranslatedYZ);
-	// Scale
-	MaterialFloat3 RotatedTranslatedScaled = (RotatedTranslated / (ExternalExpressions.Scale));
-	MaterialFloat Shape = 0;
-	
-#if CLIP_PIXELS_OUTSIDE_AABB
-	// Clip pixel if it is outside the bounding box.
-	{
-		MaterialFloat3 ComparisonShape_3D = max(RotatedTranslatedScaled, -RotatedTranslatedScaled);
-		float ComparisonShape = max(max(ComparisonShape_3D.r, ComparisonShape_3D.g), ComparisonShape_3D.b);
-		clip(1. - ComparisonShape*(1. - ExternalExpressions.Invert));
-	}
-#endif
-
-#if SHAPE_TYPE == SHAPE_TYPE_BOX
-	// Cube
-	{
-		MaterialFloat3 boxBoundsMax = smoothstep(ExternalExpressions.Inner, ExternalExpressions.Outer, RotatedTranslatedScaled);
-		MaterialFloat3 boxBoundsMin = smoothstep(ExternalExpressions.Inner, ExternalExpressions.Outer, -RotatedTranslatedScaled);
-		MaterialFloat3 result = max(boxBoundsMax, boxBoundsMin);
-		Shape = max(max(result.r, result.g), result.b);
-	}
-#elif SHAPE_TYPE == SHAPE_TYPE_CYLINDER
-	// cylinder
-	{
-		MaterialFloat Distance = length(RotatedTranslatedScaled.xy);
-		MaterialFloat Local35 = smoothstep(ExternalExpressions.Inner,ExternalExpressions.Outer, Distance);
-		MaterialFloat Local36 = smoothstep(ExternalExpressions.Inner,ExternalExpressions.Outer, RotatedTranslatedScaled.b);
-		MaterialFloat Local39 = smoothstep(ExternalExpressions.Inner,ExternalExpressions.Outer, -RotatedTranslatedScaled.b);
-		Shape = max(max(Local35,Local36),Local39);
-	}
-#elif SHAPE_TYPE == SHAPE_TYPE_CONE
-	// cone
-	{
-		MaterialFloat local33 = length(RotatedTranslatedScaled.xy);
-		MaterialFloat local36 = saturate(-RotatedTranslatedScaled.b);
-		MaterialFloat local39 = smoothstep(ExternalExpressions.Inner * local36, ExternalExpressions.Outer * local36, local33);
-		MaterialFloat local40 = smoothstep(ExternalExpressions.Inner, ExternalExpressions.Outer, RotatedTranslatedScaled.b);
-		MaterialFloat local42 = smoothstep(ExternalExpressions.Inner, ExternalExpressions.Outer, -RotatedTranslatedScaled.b);
-		Shape = max(max(local39,local40),local42);
-	}
-#elif SHAPE_TYPE == SHAPE_TYPE_SPHERE
-	// Sphere
-	{
-		MaterialFloat Distance = length(RotatedTranslatedScaled);
-		Shape = smoothstep(ExternalExpressions.Inner, ExternalExpressions.Outer, Distance);
-	}
-#endif
-	MaterialFloat ShapeAlpha = (1.00000000 - Shape);
-	
-	// Final check to reduce the performance based on shape.
-	if (ShapeAlpha < 0.00001 && !ExternalExpressions.Invert)
-	{
-#if CCR_SHADER_DISPLAY_BOUNDING_RECT
-		return float4(1,0,0,0.4);
-#else
-		return float4(0,0,0,0);
-#endif
-	}
-	MaterialFloat3 WorkingColor = RGBtoAP1(ExternalExpressions.SceneTexture);
+	MaterialFloat3 WorkingColor = sRGBtoAP1(ExternalExpressions.SceneTexture);
 	MaterialFloat3 BasicColorGrade = ComposureColorGradePerRange(WorkingColor, ExternalExpressions.ColorSaturation, ExternalExpressions.ColorContrast, ExternalExpressions.ColorGamma, ExternalExpressions.ColorGain, ExternalExpressions.ColorOffset);
 	MaterialFloat3 FinalColor = BasicColorGrade;
 #if ADVANCED_CC	
 	{
-		#define HIGHLIGHT_INTERPOLATION_THRESHOLD_WIDTH 1.0
+#define HIGHLIGHT_INTERPOLATION_THRESHOLD_WIDTH 1.0
 		MaterialFloat3 Local3 = ComposureColorGradePerRange(BasicColorGrade, ExternalExpressions.ColorSaturationShadows, ExternalExpressions.ColorContrastShadows, ExternalExpressions.ColorGammaShadows, ExternalExpressions.ColorGainShadows, ExternalExpressions.ColorOffsetShadows);
-		MaterialFloat BasicColorGrade_Mul = dot(BasicColorGrade, MaterialFloat3(0.30000001,0.58999997,0.11000000));
+		MaterialFloat BasicColorGrade_Mul = dot(BasicColorGrade, MaterialFloat3(0.30000001, 0.58999997, 0.11000000));
 		// In the future we might want to add a threshold width to increase or reduce interpolation area between shadows and midtones.
 		MaterialFloat Local5 = smoothstep(0.0 /*ExternalExpressions.ShadowMax - SHADOW_INTERPOLATION_THRESHOLD_WIDTH*/, ExternalExpressions.ShadowMax, BasicColorGrade_Mul);
 		MaterialFloat Local6 = (1.00000000 - Local5);
 		MaterialFloat3 Local7 = (Local3 * Local6);
-		MaterialFloat3 Local8 = ComposureColorGradePerRange(BasicColorGrade,ExternalExpressions.ColorSaturationMidtones, ExternalExpressions.ColorContrastMidtones, ExternalExpressions.ColorGammaMidtones, ExternalExpressions.ColorGainMidtones, ExternalExpressions.ColorOffsetMidtones);
+		MaterialFloat3 Local8 = ComposureColorGradePerRange(BasicColorGrade, ExternalExpressions.ColorSaturationMidtones, ExternalExpressions.ColorContrastMidtones, ExternalExpressions.ColorGammaMidtones, ExternalExpressions.ColorGainMidtones, ExternalExpressions.ColorOffsetMidtones);
 		// Interpolation threshold width specifies the upper bound for the soft interpolated area between midtones and highlights. 
 		// In the future we want to add control that allows user to control this width.
 		MaterialFloat Local9 = smoothstep(ExternalExpressions.HighlightsMin, ExternalExpressions.HighlightsMin + HIGHLIGHT_INTERPOLATION_THRESHOLD_WIDTH, BasicColorGrade_Mul);
@@ -667,20 +102,20 @@ float4 ProcessRegion(in out FMaterialPixelParameters Parameters, in FExternalExp
 		MaterialFloat Local11 = (1.00000000 - Local10);
 		MaterialFloat3 Local12 = (Local8 * Local11);
 		MaterialFloat3 Local13 = (Local7 + Local12);
-		MaterialFloat3 Local14 = ComposureColorGradePerRange(BasicColorGrade,ExternalExpressions.ColorSaturationHighlights, ExternalExpressions.ColorContrastHighlights, ExternalExpressions.ColorGammaHighlights, ExternalExpressions.ColorGainHighlights, ExternalExpressions.ColorOffsetHighlights);
+		MaterialFloat3 Local14 = ComposureColorGradePerRange(BasicColorGrade, ExternalExpressions.ColorSaturationHighlights, ExternalExpressions.ColorContrastHighlights, ExternalExpressions.ColorGammaHighlights, ExternalExpressions.ColorGainHighlights, ExternalExpressions.ColorOffsetHighlights);
 		MaterialFloat3 Local15 = (Local14 * Local9);
 		FinalColor = (Local13 + Local15);
 	}
 #endif
-	FinalColor = AP1toRGB(FinalColor);
-	
+	FinalColor = AP1toSRGB(FinalColor);
+
 	// White temperature calculation.
 	MaterialFloat3 AfterTemperature = FinalColor;
 
 #if TEMPERATURE_TYPE == TEMPERATURE_TYPE_LEGACY
 	MaterialFloat Local4 = smoothstep(0.00000000, 13000.00000000, ExternalExpressions.WhiteTemp);
-	MaterialFloat Local5 = (Local4 * 2.00000000  - 1.00000000);
-	
+	MaterialFloat Local5 = (Local4 * 2.00000000 - 1.00000000);
+
 	// MF Artistic temperature 
 	{
 		MaterialFloat Local7 = (Local5 * 3.00000000);
@@ -704,37 +139,145 @@ float4 ProcessRegion(in out FMaterialPixelParameters Parameters, in FExternalExp
 	AfterTemperature = WhiteBalance(ExternalExpressions.WhiteTemp, FinalColor, false /*SwapSourceAndD65White*/);
 #elif TEMPERATURE_TYPE == TEMPERATURE_TYPE_COLOR_TEMP
 	AfterTemperature = WhiteBalance(ExternalExpressions.WhiteTemp, FinalColor, true /*SwapSourceAndD65White*/);
+#endif
+	return AfterTemperature;
+}
+
+#if SHAPE_TYPE == SHAPE_TYPE_PLANE
+float GetShapeAlpha(in out FMaterialPixelParameters Parameters, in FExternalExpressions ExternalExpressions)
+{
+	// Pixel position in plane space.
+	float3 RotatedTranslatedScaled = RotateTranslateScale(GetWorldPosition(Parameters), ExternalExpressions.Translate.xyz, ExternalExpressions.Rotate, ExternalExpressions.Scale);
+
+	// Camera position in plane space.
+	float3 CameraVectorPlaneSpace = RotateTranslateScale(LWCHackToFloat(ResolvedView.WorldCameraOrigin).xyz, ExternalExpressions.Translate.xyz, ExternalExpressions.Rotate, ExternalExpressions.Scale);
+
+	float3 PositionOnPlane;
+	float DirectionSign;
+
+	float IntersectionSide;
+	{
+		float3 PlaneNormal = float3(0., 0., -1.);
+		float3 CamToPos = RotatedTranslatedScaled - CameraVectorPlaneSpace;
+		float3 DirVector = normalize(CamToPos);
+		DirectionSign = sign(dot(-PlaneNormal, DirVector));
+		PlaneNormal *= DirectionSign;
+		IntersectionSide = dot(-CameraVectorPlaneSpace, PlaneNormal);
+
+		PositionOnPlane = CameraVectorPlaneSpace + IntersectionSide / dot(DirVector, PlaneNormal) * DirVector;
+	}
+
+	float2 Interpolator = smoothstep(ExternalExpressions.Outer, ExternalExpressions.Inner, abs(PositionOnPlane.rg)) * step(-RotatedTranslatedScaled.b * DirectionSign, 0.) * step(IntersectionSide, 0.);
+	float ShapeAlpha = Interpolator.r * step(0., Interpolator.g) * Interpolator.g * step(0., Interpolator.r);
+
+#if CLIP_PIXELS_OUTSIDE_AABB
+	// Clip pixel if it is outside the bounding box.
+	{
+		clip(ShapeAlpha  + ExternalExpressions.Invert - 0.00001);
+	}
+#endif
+
+	return ShapeAlpha;
+}
+
+#else
+
+float GetShapeAlpha(in out FMaterialPixelParameters Parameters, in FExternalExpressions ExternalExpressions)
+{
+	float3 RotatedTranslatedScaled = RotateTranslateScale(GetWorldPosition(Parameters), ExternalExpressions.Translate.xyz, ExternalExpressions.Rotate, ExternalExpressions.Scale);
+
+	float Shape = 0;
+
+#if CLIP_PIXELS_OUTSIDE_AABB
+	// Clip pixel if it is outside the bounding box.
+	{
+		float3 ComparisonShape_3D = max(RotatedTranslatedScaled, -RotatedTranslatedScaled);
+		float ComparisonShape = max(max(ComparisonShape_3D.r, ComparisonShape_3D.g), ComparisonShape_3D.b);
+		clip(1. - ComparisonShape * (1. - ExternalExpressions.Invert));
+	}
+#endif
+
+#if SHAPE_TYPE == SHAPE_TYPE_BOX
+	// Cube
+	{
+		float3 boxBoundsMax = smoothstep(ExternalExpressions.Inner, ExternalExpressions.Outer, RotatedTranslatedScaled);
+		float3 boxBoundsMin = smoothstep(ExternalExpressions.Inner, ExternalExpressions.Outer, -RotatedTranslatedScaled);
+		float3 result = max(boxBoundsMax, boxBoundsMin);
+		Shape = max(max(result.r, result.g), result.b);
+	}
+#elif SHAPE_TYPE == SHAPE_TYPE_CYLINDER
+	// cylinder
+	{
+		float Distance = length(RotatedTranslatedScaled.xy);
+		float Local35 = smoothstep(ExternalExpressions.Inner, ExternalExpressions.Outer, Distance);
+		float Local36 = smoothstep(ExternalExpressions.Inner, ExternalExpressions.Outer, RotatedTranslatedScaled.b);
+		float Local39 = smoothstep(ExternalExpressions.Inner, ExternalExpressions.Outer, -RotatedTranslatedScaled.b);
+		Shape = max(max(Local35, Local36), Local39);
+	}
+#elif SHAPE_TYPE == SHAPE_TYPE_CONE
+	// cone
+	{
+		float local33 = length(RotatedTranslatedScaled.xy);
+		float local36 = saturate(-RotatedTranslatedScaled.b);
+		float local39 = smoothstep(ExternalExpressions.Inner * local36, ExternalExpressions.Outer * local36, local33);
+		float local40 = smoothstep(ExternalExpressions.Inner, ExternalExpressions.Outer, RotatedTranslatedScaled.b);
+		float local42 = smoothstep(ExternalExpressions.Inner, ExternalExpressions.Outer, -RotatedTranslatedScaled.b);
+		Shape = max(max(local39, local40), local42);
+	}
+#elif SHAPE_TYPE == SHAPE_TYPE_SPHERE
+	// Sphere
+	{
+		float Distance = length(RotatedTranslatedScaled);
+		Shape = smoothstep(ExternalExpressions.Inner, ExternalExpressions.Outer, Distance);
+	}
+#endif
+	float ShapeAlpha = (1.00000000 - Shape);
+	return ShapeAlpha;
+}
 #endif
 
+MaterialFloat4 ProcessRegion(in out FMaterialPixelParameters Parameters, in FExternalExpressions ExternalExpressions)
+{
+	float ShapeAlpha = GetShapeAlpha(Parameters, ExternalExpressions);
+
+	// Final check to reduce the performance based on shape.
+	if (ShapeAlpha < 0.00001 && !ExternalExpressions.Invert)
+	{
+#if CCR_SHADER_DISPLAY_BOUNDING_RECT
+		return float4(1, 0, 0, 0.4);
+#else
+		return float4(0, 0, 0, 0);
+#endif
+	}
+
+	float3 AfterTemperature = ProcessColor(ExternalExpressions);
+
 	/* Falloff */
-	MaterialFloat ShapeWithFalloff = PositiveClampedPow(ShapeAlpha, ExternalExpressions.Falloff);
-	MaterialFloat ShapeWithFalloffInverse = (1.00000000 - ShapeWithFalloff);
+	float ShapeWithFalloff = PositiveClampedPow(ShapeAlpha, ExternalExpressions.Falloff);
+	float ShapeWithFalloffInverse = (1.00000000 - ShapeWithFalloff);
 
-	MaterialFloat FinalShape = lerp(ShapeWithFalloff, ShapeWithFalloffInverse, ExternalExpressions.Invert);
+	float FinalShape = lerp(ShapeWithFalloff, ShapeWithFalloffInverse, ExternalExpressions.Invert);
 
-	MaterialFloat FinalAlpha = lerp(0.00000000, ExternalExpressions.Intensity, FinalShape);
-	
-	FinalColor = lerp(ExternalExpressions.SceneTexture.rgb, AfterTemperature, FinalAlpha);
+	float FinalAlpha = lerp(0.00000000, ExternalExpressions.Intensity, FinalShape);
+
+	MaterialFloat3 FinalColor = lerp(ExternalExpressions.SceneTexture.rgb, AfterTemperature, FinalAlpha);
 
 	if (ExternalExpressions.ExcludeStencil > 0)
 	{
-		MaterialFloat4 CustomStencil = SceneTextureLookup(GetDefaultSceneTextureUV(Parameters, 25), 25, false);
+		float4 CustomStencil = SceneTextureLookup(GetDefaultSceneTextureUV(Parameters, PPI_CustomStencil), PPI_CustomStencil, false);
 		FinalColor = lerp(FinalColor, ExternalExpressions.SceneTexture.rgb, clamp(CustomStencil.r, 0., 1.));
 	}
 
-	return float4(FinalColor, FinalAlpha);
+	return MaterialFloat4(FinalColor, FinalAlpha);
 }
 
-
 void CalcPixelMaterialInputs(in out FMaterialPixelParameters Parameters, in out FPixelMaterialInputs PixelMaterialInputs)
 {
-    // Initial calculations (required for Normal)
-
     // The Normal is a special case as it might have its own expressions and also be used to calculate other inputs, so perform the assignment here
 	PixelMaterialInputs.Normal = MaterialFloat3(0.00000000, 0.00000000, 1.00000000);
 
     // Note that here MaterialNormal can be in world space or tangent space
-	float3 MaterialNormal = GetMaterialNormal(Parameters, PixelMaterialInputs);
+	half3 MaterialNormal = GetMaterialNormal(Parameters, PixelMaterialInputs);
 
 	Parameters.WorldNormal = normalize(MaterialNormal);
 
@@ -787,8 +330,8 @@ void CalcPixelMaterialInputs(in out FMaterialPixelParameters Parameters, in out
 			ExtExpressions.HighlightsMin = ColorCorrectHighlights.HighlightsMin;
 		}
 #endif
-		MaterialFloat4 RegionColor = ProcessRegion(Parameters, ExtExpressions);
 
+		MaterialFloat4 RegionColor = ProcessRegion(Parameters, ExtExpressions);
 		FinalColor = lerp(FinalColor.rgba, RegionColor.rgba, RegionColor.a);
 	}
 		
@@ -796,12 +339,7 @@ void CalcPixelMaterialInputs(in out FMaterialPixelParameters Parameters, in out
 #if CLIP_PIXELS_OUTSIDE_AABB
 	PixelMaterialInputs.Opacity = FinalColor.a;
 #else
-	#if SAMPLE_OPACITY_FROM_GBUFFER
-		// On the lower scalability settings there is no alpha channel in Scene texture.
-		PixelMaterialInputs.Opacity = 1.0;
-	#else
-		PixelMaterialInputs.Opacity = SceneTexture.a;
-	#endif
+	PixelMaterialInputs.Opacity = 1.0;
 #endif
 
 	PixelMaterialInputs.OpacityMask = 1.00000000;
@@ -824,16 +362,6 @@ struct FPostProcessMaterialVSToPS
 	float4 Position : SV_POSITION;
 };
 
-FMaterialVertexParameters GetPostProcessMaterialVSParameters(float2 UV)
-{
-	// Most params irrelevant so not a lot to fill out here
-	FMaterialVertexParameters Result = (FMaterialVertexParameters) 0;
-	Result.VertexColor = 1.f.xxxx;
-	Result.WorldPosition = float3(UV, 0.f);
-	
-	return Result;
-}
-
 #if (FEATURE_LEVEL > FEATURE_LEVEL_ES3_1)
 void MainVS(
 	in float4 InPosition : ATTRIBUTE0,
@@ -844,39 +372,45 @@ void MainVS(
 	DrawRectangle(InPosition, Output.Position);
 
 }
-
-void MainPS(
-	in FPostProcessMaterialVSToPS Input,
-	out float4 OutColor : SV_Target0
-	)
+void BuildParameters(in FPostProcessMaterialVSToPS Input, inout FMaterialPixelParameters Parameters)
 {
+	Parameters = (FMaterialPixelParameters)0;
+
 	ResolvedView = ResolveView();
-	FMaterialPixelParameters Parameters = MakeInitializedMaterialPixelParameters();
-	FPixelMaterialInputs PixelMaterialInputs;
 
-	// can be optimized
 	float4 SvPosition = Input.Position;
 	float2 ViewportUV = (SvPosition.xy - PostProcessOutput_ViewportMin.xy) * PostProcessOutput_ViewportSizeInverse.xy;
-	
+
 	SvPosition.z = LookupDeviceZ(ViewportUVToBufferUV(ViewportUV));
 	SvPosition.z = max(SvPosition.z, 1e-18);
-	
+
 	float4 ScreenPosition = SvPositionToScreenPosition(SvPosition);
 	float3 TranslatedWorldPosition = SvPositionToTranslatedWorld(SvPosition);
 
 	// Remove the pre view translation
 	Parameters.WorldPosition_CamRelative = TranslatedWorldPosition.xyz;
-	Parameters.AbsoluteWorldPosition = LWCToFloat(LWCSubtract(TranslatedWorldPosition, ResolvedView.PreViewTranslation));
+	Parameters.AbsoluteWorldPosition = TranslatedWorldPosition.xyz - LWCHackToFloat(ResolvedView.PreViewTranslation).xyz;
 
 	Parameters.SvPosition = SvPosition;
 	Parameters.ScreenPosition = ScreenPosition;
 
 	Parameters.CameraVector = normalize(-Parameters.WorldPosition_CamRelative.xyz);
-	
+
 	/* Unused */
 	Parameters.LightVector = 0;
 	Parameters.VertexColor = 1;
 	Parameters.TwoSidedSign = 1.0f;
+}
+
+void MainPS(
+	in FPostProcessMaterialVSToPS Input,
+	out float4 OutColor : SV_Target0
+	)
+{
+	FMaterialPixelParameters Parameters;
+	FPixelMaterialInputs PixelMaterialInputs;
+
+	BuildParameters(Input, Parameters);
 
     // Now that we have all the pixel-related parameters setup, calculate the Material Input/Attributes and Normal
 	CalcPixelMaterialInputs(Parameters, PixelMaterialInputs);
@@ -886,4 +420,5 @@ void MainPS(
 	OutColor = float4(max(PixelMaterialInputs.EmissiveColor, 0.), Alpha);
 	OutColor.xyz *= View.PreExposure;
 } 
+
 #endif 
diff --git a/Engine/Plugins/Experimental/ColorCorrectRegions/Source/ColorCorrectRegions/Private/ColorCorrectRegion.cpp b/Engine/Plugins/Experimental/ColorCorrectRegions/Source/ColorCorrectRegions/Private/ColorCorrectRegion.cpp
--- a/Engine/Plugins/Experimental/ColorCorrectRegions/Source/ColorCorrectRegions/Private/ColorCorrectRegion.cpp
+++ b/Engine/Plugins/Experimental/ColorCorrectRegions/Source/ColorCorrectRegions/Private/ColorCorrectRegion.cpp
@@ -144,21 +144,29 @@ void AColorCorrectRegion::PostEditChangeProperty(struct FPropertyChangedEvent& P
 	Super::PostEditChangeProperty(PropertyChangedEvent);
 	const FName PropertyName = PropertyChangedEvent.GetPropertyName();
 
+	if (!ColorCorrectRegionsSubsystem)
+	{
+		if (const UWorld* World = GetWorld())
+		{
+			ColorCorrectRegionsSubsystem = World->GetSubsystem<UColorCorrectRegionsSubsystem>();
+		}
+	}
+
 	// Reorder all CCRs after the Priority property has changed.
 	// Also, in context of Multi-User: PropertyChangedEvent can be a stub without the actual property data. 
 	// Therefore we need to refresh priority if PropertyChangedEvent.Property is nullptr. 
 	if (PropertyName == GET_MEMBER_NAME_CHECKED(AColorCorrectRegion, Priority) || PropertyChangedEvent.Property == nullptr)
 	{
-		if (!ColorCorrectRegionsSubsystem)
+		if (ColorCorrectRegionsSubsystem)
 		{
-			if (const UWorld* World = GetWorld())
-			{
-				ColorCorrectRegionsSubsystem = World->GetSubsystem<UColorCorrectRegionsSubsystem>();
-			}
+			ColorCorrectRegionsSubsystem->SortRegionsByPriority();
 		}
+	}
+	if (PropertyName == GET_MEMBER_NAME_CHECKED(AColorCorrectRegion, Type) || PropertyChangedEvent.Property == nullptr)
+	{
 		if (ColorCorrectRegionsSubsystem)
 		{
-			ColorCorrectRegionsSubsystem->SortRegionsByPriority();
+			ColorCorrectRegionsSubsystem->OnLevelsChanged();
 		}
 	}
 	GetActorBounds(true, BoxOrigin, BoxExtent);
diff --git a/Engine/Plugins/Experimental/ColorCorrectRegions/Source/ColorCorrectRegions/Private/ColorCorrectRegionSubsystem.cpp b/Engine/Plugins/Experimental/ColorCorrectRegions/Source/ColorCorrectRegions/Private/ColorCorrectRegionSubsystem.cpp
--- a/Engine/Plugins/Experimental/ColorCorrectRegions/Source/ColorCorrectRegions/Private/ColorCorrectRegionSubsystem.cpp
+++ b/Engine/Plugins/Experimental/ColorCorrectRegions/Source/ColorCorrectRegions/Private/ColorCorrectRegionSubsystem.cpp
@@ -53,11 +53,17 @@ void UColorCorrectRegionsSubsystem::Deinitialize()
 	}
 #endif
 	GetWorld()->OnLevelsChanged().RemoveAll(this);
-	for (AColorCorrectRegion* Region : Regions)
+	for (AColorCorrectRegion* Region : RegionsPriorityBased)
 	{
 		Region->Cleanup();
 	}
-	Regions.Reset();
+	RegionsPriorityBased.Reset();	
+	
+	for (AColorCorrectRegion* Region : RegionsDistanceBased)
+	{
+		Region->Cleanup();
+	}
+	RegionsDistanceBased.Reset();
 
 	// Prevent this SVE from being gathered, in case it is kept alive by a strong reference somewhere else.
 	{
@@ -83,14 +89,19 @@ void UColorCorrectRegionsSubsystem::OnActorSpawned(AActor* InActor)
 	if (IsRegionValid(AsRegion, GetWorld()))
 	{
 		FScopeLock RegionScopeLock(&RegionAccessCriticalSection);
-
+		EColorCorrectRegionsType CCRType = AsRegion->Type;
 		// We wouldn't have to do a check here except in case of nDisplay we need to populate this list during OnLevelsChanged 
 		// because nDisplay can release Actors while those are marked as BeginningPlay. Therefore we want to avoid 
 		// adding regions twice.
-		if (!Regions.Contains(AsRegion))
+		TArray<AColorCorrectRegion*>* RegionsToAddTo = CCRType == EColorCorrectRegionsType::Plane ? &RegionsDistanceBased : &RegionsPriorityBased;
+		if (!RegionsToAddTo->Contains(AsRegion))
 		{
-			Regions.Add(AsRegion);
-			SortRegionsByPriority();
+			RegionsToAddTo->Add(AsRegion);
+			// Distance based CCR can only be sorted on render, when View info is available.
+			if (CCRType != EColorCorrectRegionsType::Plane)
+			{
+				SortRegionsByPriority();
+			}
 		}
 	}
 }
@@ -106,7 +117,8 @@ void UColorCorrectRegionsSubsystem::OnActorDeleted(AActor* InActor)
 #endif
 	{
 		FScopeLock RegionScopeLock(&RegionAccessCriticalSection);
-		Regions.Remove(AsRegion);
+		RegionsPriorityBased.Remove(AsRegion);
+		RegionsDistanceBased.Remove(AsRegion);
 	}
 }
 
@@ -114,26 +126,52 @@ void UColorCorrectRegionsSubsystem::SortRegionsByPriority()
 {
 	FScopeLock RegionScopeLock(&RegionAccessCriticalSection);
 
-	Regions.Sort([](const AColorCorrectRegion& A, const AColorCorrectRegion& B) {
+	RegionsPriorityBased.Sort([](const AColorCorrectRegion& A, const AColorCorrectRegion& B) {
 		// Regions with the same priority could potentially cause flickering on overlap
 		return A.Priority < B.Priority;
 	});
 }
 
+void UColorCorrectRegionsSubsystem::SortRegionsByDistance(const FVector& ViewLocation)
+{
+	FScopeLock RegionScopeLock(&RegionAccessCriticalSection);
+	TMap<AColorCorrectRegion*, double> DistanceMap;
+	for (AColorCorrectRegion* Region : RegionsDistanceBased)
+	{
+		FVector BoxCenter, BoxExtents;
+		Region->GetBounds(BoxCenter, BoxExtents);
+		FVector CameraToRegionVec = (BoxCenter - ViewLocation);
+		DistanceMap.Add(Region, CameraToRegionVec.Dot(CameraToRegionVec));
+	}
+
+	RegionsDistanceBased.Sort([&DistanceMap](const AColorCorrectRegion& A, const AColorCorrectRegion& B) {
+		// Regions with the same distance could potentially cause flickering on overlap
+		return DistanceMap[&A] > DistanceMap[&B];
+	});
+
+}
+
 void UColorCorrectRegionsSubsystem::RefreshRegions()
 {
 	FScopeLock RegionScopeLock(&RegionAccessCriticalSection);
 
-	Regions.Reset();
+	RegionsPriorityBased.Reset();
+	RegionsDistanceBased.Reset();
 	for (TActorIterator<AColorCorrectRegion> It(GetWorld()); It; ++It)
 	{
 		AColorCorrectRegion* AsRegion = *It;
 		if (IsRegionValid(AsRegion, GetWorld()))
 		{
-			Regions.Add(AsRegion);
+			if (AsRegion->Type != EColorCorrectRegionsType::Plane)
+			{
+				RegionsPriorityBased.Add(AsRegion);
+			}
+			else
+			{
+				RegionsDistanceBased.Add(AsRegion);
+			}
 		}
 	}
 	SortRegionsByPriority();
 }
 
-
diff --git a/Engine/Plugins/Experimental/ColorCorrectRegions/Source/ColorCorrectRegions/Private/ColorCorrectRegionsSceneViewExtension.cpp b/Engine/Plugins/Experimental/ColorCorrectRegions/Source/ColorCorrectRegions/Private/ColorCorrectRegionsSceneViewExtension.cpp
--- a/Engine/Plugins/Experimental/ColorCorrectRegions/Source/ColorCorrectRegions/Private/ColorCorrectRegionsSceneViewExtension.cpp
+++ b/Engine/Plugins/Experimental/ColorCorrectRegions/Source/ColorCorrectRegions/Private/ColorCorrectRegionsSceneViewExtension.cpp
@@ -1,667 +1,661 @@
-// Copyright Epic Games, Inc. All Rights Reserved.
-
-#include "ColorCorrectRegionsSceneViewExtension.h"
-#include "RHI.h"
-#include "SceneView.h"
-#include "PostProcess/PostProcessing.h"
-#include "ColorCorrectRegionsModule.h"
-#include "ColorCorrectRegionsSubsystem.h"
-#include "ColorCorrectRegionsPostProcessMaterial.h"
-#include "ScreenPass.h"
-#include "CommonRenderResources.h"
-#include "Containers/DynamicRHIResourceArray.h"
-#include "Engine/World.h"
-#include "EngineUtils.h"
-#include "SceneRendering.h"
-#include "DynamicResolutionState.h"
-// Set this to 1 to clip pixels outside of bounding box.
-#define CLIP_PIXELS_OUTSIDE_AABB 1
-
-//Set this to 1 to see the clipping region.
-#define ColorCorrectRegions_SHADER_DISPLAY_BOUNDING_RECT 0
-
-
-namespace
-{
-	FRHIDepthStencilState* GetMaterialStencilState(const FMaterial* Material)
-	{
-		static FRHIDepthStencilState* StencilStates[] =
-		{
-			TStaticDepthStencilState<false, CF_Always, true, CF_Less>::GetRHI(),
-			TStaticDepthStencilState<false, CF_Always, true, CF_LessEqual>::GetRHI(),
-			TStaticDepthStencilState<false, CF_Always, true, CF_Greater>::GetRHI(),
-			TStaticDepthStencilState<false, CF_Always, true, CF_GreaterEqual>::GetRHI(),
-			TStaticDepthStencilState<false, CF_Always, true, CF_Equal>::GetRHI(),
-			TStaticDepthStencilState<false, CF_Always, true, CF_NotEqual>::GetRHI(),
-			TStaticDepthStencilState<false, CF_Always, true, CF_Never>::GetRHI(),
-			TStaticDepthStencilState<false, CF_Always, true, CF_Always>::GetRHI(),
-		};
-		static_assert(EMaterialStencilCompare::MSC_Count == UE_ARRAY_COUNT(StencilStates), "Ensure that all EMaterialStencilCompare values are accounted for.");
-
-		check(Material);
-
-		return StencilStates[Material->GetStencilCompare()];
-	}
-
-	FScreenPassTextureViewportParameters GetTextureViewportParameters(const FScreenPassTextureViewport& InViewport)
-	{
-		const FVector2f Extent(InViewport.Extent);	// LWC_TODO: Precision loss
-		const FVector2f ViewportMin(InViewport.Rect.Min.X, InViewport.Rect.Min.Y);
-		const FVector2f ViewportMax(InViewport.Rect.Max.X, InViewport.Rect.Max.Y);
-		const FVector2f ViewportSize = ViewportMax - ViewportMin;
-
-		FScreenPassTextureViewportParameters Parameters;
-
-		if (!InViewport.IsEmpty())
-		{
-			Parameters.Extent = FVector2f(Extent);
-			Parameters.ExtentInverse = FVector2f(1.0f / Extent.X, 1.0f / Extent.Y);
-
-			Parameters.ScreenPosToViewportScale = FVector2f(0.5f, -0.5f) * ViewportSize;	
-			Parameters.ScreenPosToViewportBias = (0.5f * ViewportSize) + ViewportMin;	
-
-			Parameters.ViewportMin = InViewport.Rect.Min;
-			Parameters.ViewportMax = InViewport.Rect.Max;
-
-			Parameters.ViewportSize = ViewportSize;
-			Parameters.ViewportSizeInverse = FVector2f(1.0f / Parameters.ViewportSize.X, 1.0f / Parameters.ViewportSize.Y);
-
-			Parameters.UVViewportMin = ViewportMin * Parameters.ExtentInverse;
-			Parameters.UVViewportMax = ViewportMax * Parameters.ExtentInverse;
-
-			Parameters.UVViewportSize = Parameters.UVViewportMax - Parameters.UVViewportMin;
-			Parameters.UVViewportSizeInverse = FVector2f(1.0f / Parameters.UVViewportSize.X, 1.0f / Parameters.UVViewportSize.Y);
-
-			Parameters.UVViewportBilinearMin = Parameters.UVViewportMin + 0.5f * Parameters.ExtentInverse;
-			Parameters.UVViewportBilinearMax = Parameters.UVViewportMax - 0.5f * Parameters.ExtentInverse;
-		}
-
-		return Parameters;
-	}
-
-
-	void GetPixelSpaceBoundingRect(const FSceneView& InView, const FVector& InBoxCenter, const FVector& InBoxExtents, FIntRect& OutViewport, float& OutMaxDepth, float& OutMinDepth)
-	{
-		OutViewport = FIntRect(INT32_MAX, INT32_MAX, -INT32_MAX, -INT32_MAX);
-		// 8 corners of the bounding box. To be multiplied by box extent and offset by the center.
-		const int NumCorners = 8;
-		const FVector Verts[NumCorners] = {
-			FVector(1, 1, 1),
-			FVector(1, 1,-1),
-			FVector(1,-1, 1),
-			FVector(1,-1,-1),
-			FVector(-1, 1, 1),
-			FVector(-1, 1,-1),
-			FVector(-1,-1, 1),
-			FVector(-1,-1,-1) };
-
-		for (int32 Index = 0; Index < NumCorners; Index++)
-		{
-			// Project bounding box vertecies into screen space.
-			const FVector WorldVert = InBoxCenter + (Verts[Index] * InBoxExtents);
-			FVector2D PixelVert;
-			FVector4 ScreenSpaceCoordinate = InView.WorldToScreen(WorldVert);
-
-			OutMaxDepth = FMath::Max<float>(ScreenSpaceCoordinate.W, OutMaxDepth);
-			OutMinDepth = FMath::Min<float>(ScreenSpaceCoordinate.W, OutMinDepth);
-
-			if (InView.ScreenToPixel(ScreenSpaceCoordinate, PixelVert))
-			{
-				// Update screen-space bounding box with with transformed vert.
-				OutViewport.Min.X = FMath::Min<int32>(OutViewport.Min.X, PixelVert.X);
-				OutViewport.Min.Y = FMath::Min<int32>(OutViewport.Min.Y, PixelVert.Y);
-
-				OutViewport.Max.X = FMath::Max<int32>(OutViewport.Max.X, PixelVert.X);
-				OutViewport.Max.Y = FMath::Max<int32>(OutViewport.Max.Y, PixelVert.Y);
-			}
-		}
-	}
-
-	// Function that calculates all points of intersection between plane and bounding box. Resulting points are unsorted.
-	void CalculatePlaneAABBIntersectionPoints(const FPlane& Plane, const FVector& BoxCenter, const FVector& BoxExtents, TArray<FVector>& OutPoints)
-	{
-		const FVector MaxCorner = BoxCenter + BoxExtents;
-
-		const FVector Verts[3][4] = {
-			{
-				// X Direction
-				FVector(-1, -1, -1),
-				FVector(-1,  1, -1),
-				FVector(-1, -1,  1),
-				FVector(-1,  1,  1),
-			},
-			{
-				// Y Direction
-				FVector(-1, -1, -1),
-				FVector( 1, -1, -1),
-				FVector( 1, -1,  1),
-				FVector(-1, -1,  1),
-			},
-			{
-				// Z Direction
-				FVector(-1, -1, -1),
-				FVector( 1, -1, -1),
-				FVector( 1,  1, -1),
-				FVector(-1,  1, -1),
-			}
-		};
-
-		FVector Intersection;
-		FVector Start;
-		FVector End;
-
-		for (int RunningAxis_Dir = 0; RunningAxis_Dir < 3; RunningAxis_Dir++)
-		{
-			const FVector *CornerLocations = Verts[RunningAxis_Dir];
-			for (int RunningCorner = 0; RunningCorner < 4; RunningCorner++)
-			{
-				Start = BoxCenter + BoxExtents * CornerLocations[RunningCorner];
-				End = FVector(Start.X, Start.Y, Start.Z);
-				End[RunningAxis_Dir] = MaxCorner[RunningAxis_Dir];
-				if (FMath::SegmentPlaneIntersection(Start, End, Plane, Intersection))
-				{
-					OutPoints.Add(Intersection);
-				}
-			}
-		}
-	}
-
-	// Takes in an existing viewport and updates it with an intersection bounding rectangle.
-	void UpdateMinMaxWithFrustrumAABBIntersection(const FSceneView& InView, const FVector& InBoxCenter, const FVector& InBoxExtents, FIntRect& OutViewportToUpdate, float& OutMaxDepthToUpdate)
-	{
-		TArray<FVector> Points;
-		Points.Reserve(6);
-		static bool bNotifiedOfClippingPlaneError = false;
-
-		if (InView.bHasNearClippingPlane)
-		{
-			CalculatePlaneAABBIntersectionPoints(InView.NearClippingPlane, InBoxCenter, InBoxExtents, Points);
-		}
-		// Previously last plane was near clipping plane.
-		else if (InView.ViewFrustum.Planes.Num() == 5)
-		{
-			CalculatePlaneAABBIntersectionPoints(InView.ViewFrustum.Planes[4], InBoxCenter, InBoxExtents, Points);
-		}
-		else if (!bNotifiedOfClippingPlaneError)
-		{
-			bNotifiedOfClippingPlaneError = true;
-			UE_LOG(ColorCorrectRegions, Error, TEXT("Couldn't find a correct near clipping plane in View Frustrum"));
-		}
-
-		if (Points.Num() == 0)
-		{
-			return;
-		}
-
-		for (FVector Point : Points)
-		{
-			// Project bounding box vertecies into screen space.
-			FVector4 ScreenSpaceCoordinate = InView.WorldToScreen(Point);
-			FVector4 ScreenSpaceCoordinateScaled = ScreenSpaceCoordinate * 1.0 / ScreenSpaceCoordinate.W;
-
-			OutMaxDepthToUpdate = FMath::Max<float>(ScreenSpaceCoordinate.W, OutMaxDepthToUpdate);
-			FVector2D PixelVert;
-
-			if (InView.ScreenToPixel(ScreenSpaceCoordinate, PixelVert))
-			{
-				// Update screen-space bounding box with with transformed vert.
-				OutViewportToUpdate.Min.X = FMath::Min<int32>(OutViewportToUpdate.Min.X, PixelVert.X);
-				OutViewportToUpdate.Min.Y = FMath::Min<int32>(OutViewportToUpdate.Min.Y, PixelVert.Y);
-
-				OutViewportToUpdate.Max.X = FMath::Max<int32>(OutViewportToUpdate.Max.X, PixelVert.X);
-				OutViewportToUpdate.Max.Y = FMath::Max<int32>(OutViewportToUpdate.Max.Y, PixelVert.Y);
-			}
-		}
-	}
-
-	bool ViewSupportsRegions(const FSceneView& View)
-	{
-		return View.Family->EngineShowFlags.PostProcessing &&
-				View.Family->EngineShowFlags.PostProcessMaterial;
-	}
-
-	template<typename TSetupFunction>
-	void DrawScreenPass(
-		FRHICommandList& RHICmdList,
-		const FSceneView& View,
-		const FScreenPassTextureViewport& OutputViewport,
-		const FScreenPassTextureViewport& InputViewport,
-		const FScreenPassPipelineState& PipelineState,
-		TSetupFunction SetupFunction)
-	{
-		PipelineState.Validate();
-
-		const FIntRect InputRect = InputViewport.Rect;
-		const FIntPoint InputSize = InputViewport.Extent;
-		const FIntRect OutputRect = OutputViewport.Rect;
-		const FIntPoint OutputSize = OutputRect.Size();
-
-		RHICmdList.SetViewport(OutputRect.Min.X, OutputRect.Min.Y, 0.0f, OutputRect.Max.X, OutputRect.Max.Y, 1.0f);
-
-		SetScreenPassPipelineState(RHICmdList, PipelineState);
-
-		// Setting up buffers.
-		SetupFunction(RHICmdList);
-
-		FIntPoint LocalOutputPos(FIntPoint::ZeroValue);
-		FIntPoint LocalOutputSize(OutputSize);
-		EDrawRectangleFlags DrawRectangleFlags = EDRF_UseTriangleOptimization;
-
-		DrawPostProcessPass(
-			RHICmdList,
-			LocalOutputPos.X, LocalOutputPos.Y, LocalOutputSize.X, LocalOutputSize.Y,
-			InputRect.Min.X, InputRect.Min.Y, InputRect.Width(), InputRect.Height(),
-			OutputSize,
-			InputSize,
-			PipelineState.VertexShader,
-			View.StereoViewIndex,
-			false,
-			DrawRectangleFlags);
-	}
-
-	// A helper function for getting the right shader.
-	TShaderMapRef<FColorCorrectRegionMaterialPS> GetRegionShader(const FGlobalShaderMap* GlobalShaderMap, EColorCorrectRegionsType RegionType, FColorCorrectRegionMaterialPS::ETemperatureType TemperatureType, bool bIsAdvanced, bool bSampleOpacityFromGbuffer)
-	{
-		FColorCorrectRegionMaterialPS::FPermutationDomain PermutationVector;
-		PermutationVector.Set<FColorCorrectRegionMaterialPS::FAdvancedShader>(bIsAdvanced);
-		PermutationVector.Set<FColorCorrectRegionMaterialPS::FShaderType>(RegionType);
-		PermutationVector.Set<FColorCorrectRegionMaterialPS::FTemperatureType>(TemperatureType);
-
-#if ColorCorrectRegions_SHADER_DISPLAY_BOUNDING_RECT
-		PermutationVector.Set<FColorCorrectRegionMaterialPS::FDisplayBoundingRect>(true);
-#endif
-#if CLIP_PIXELS_OUTSIDE_AABB
-		PermutationVector.Set<FColorCorrectRegionMaterialPS::FClipPixelsOutsideAABB>(true);
-#endif
-		PermutationVector.Set<FColorCorrectRegionMaterialPS::FSampleOpacityFromGbuffer>(bSampleOpacityFromGbuffer);
-		
-		return TShaderMapRef<FColorCorrectRegionMaterialPS>(GlobalShaderMap, PermutationVector);
-		;
-	}
-
-	FVector4 Clamp(const FVector4 & VectorToClamp, float Min, float Max)
-	{
-		return FVector4(FMath::Clamp(VectorToClamp.X, Min, Max),
-						FMath::Clamp(VectorToClamp.Y, Min, Max),
-						FMath::Clamp(VectorToClamp.Z, Min, Max),
-						FMath::Clamp(VectorToClamp.W, Min, Max));
-	}
-}
-
-FColorCorrectRegionsSceneViewExtension::FColorCorrectRegionsSceneViewExtension(const FAutoRegister& AutoRegister, UColorCorrectRegionsSubsystem* InWorldSubsystem) :
-	FSceneViewExtensionBase(AutoRegister), WorldSubsystem(InWorldSubsystem)
-{
-}
-
-void FColorCorrectRegionsSceneViewExtension::PrePostProcessPass_RenderThread(FRDGBuilder& GraphBuilder, const FSceneView& View, const FPostProcessingInputs& Inputs)
-{
-	// Necessary for when an actor is added or removed from the scene. Also when priority is changed.
-	FScopeLock RegionScopeLock(&WorldSubsystem->RegionAccessCriticalSection);
-
-	if (WorldSubsystem->Regions.Num() == 0 || !ViewSupportsRegions(View))
-	{
-		return;
-	}
-
-	Inputs.Validate();
-
-	const FSceneViewFamily& ViewFamily = *View.Family;
-
-	DynamicRenderScaling::TMap<float> UpperBounds = ViewFamily.GetScreenPercentageInterface()->GetResolutionFractionsUpperBound();
-
-	const auto FeatureLevel = View.GetFeatureLevel();
-	const float ScreenPercentage = UpperBounds[GDynamicPrimaryResolutionFraction] * ViewFamily.SecondaryViewFraction;
-	
-	// We need to make sure to take Windows and Scene scale into account.
-
-	checkSlow(View.bIsViewInfo); // can't do dynamic_cast because FViewInfo doesn't have any virtual functions.
-	const FIntRect PrimaryViewRect = static_cast<const FViewInfo&>(View).ViewRect;
-
-	FScreenPassTexture SceneColor((*Inputs.SceneTextures)->SceneColorTexture, PrimaryViewRect);
-
-
-	if (!SceneColor.IsValid())
-	{
-		return;
-	}
-
-	{
-		// Getting material data for the current view.
-		FGlobalShaderMap* GlobalShaderMap = GetGlobalShaderMap(GMaxRHIFeatureLevel);
-
-		// Reusing the same output description for our back buffer as SceneColor
-		FRDGTextureDesc ColorCorrectRegionsOutputDesc = SceneColor.Texture->Desc;
-		bool bSampleOpacityFromGbuffer = false;
-		if (ColorCorrectRegionsOutputDesc.Format != PF_FloatRGBA)
-		{
-			bSampleOpacityFromGbuffer = true;
-		}
-		ColorCorrectRegionsOutputDesc.Format = PF_FloatRGBA;
-		FLinearColor ClearColor(0., 0., 0., 0.);
-		ColorCorrectRegionsOutputDesc.ClearValue = FClearValueBinding(ClearColor);
-
-		FRDGTexture* BackBufferRenderTargetTexture = GraphBuilder.CreateTexture(ColorCorrectRegionsOutputDesc, TEXT("BackBufferRenderTargetTexture"));
-		FScreenPassRenderTarget BackBufferRenderTarget = FScreenPassRenderTarget(BackBufferRenderTargetTexture, SceneColor.ViewRect, ERenderTargetLoadAction::EClear);
-		FScreenPassRenderTarget SceneColorRenderTarget(SceneColor, ERenderTargetLoadAction::ELoad);
-		const FScreenPassTextureViewport SceneColorTextureViewport(SceneColor);
-
-		FRHIBlendState* DefaultBlendState = FScreenPassPipelineState::FDefaultBlendState::GetRHI();
-		FRHIDepthStencilState* DepthStencilState = FScreenPassPipelineState::FDefaultDepthStencilState::GetRHI();
-
-		RDG_EVENT_SCOPE(GraphBuilder, "Color Correct Regions %dx%d", SceneColorTextureViewport.Rect.Width(), SceneColorTextureViewport.Rect.Height());
-
-		FRHISamplerState* PointClampSampler = TStaticSamplerState<SF_Point, AM_Clamp, AM_Clamp, AM_Clamp>::GetRHI();
-		const FScreenPassTextureViewportParameters SceneTextureViewportParams = GetTextureViewportParameters(SceneColorTextureViewport);
-		FScreenPassTextureInput SceneTextureInput;
-		{
-			SceneTextureInput.Viewport = SceneTextureViewportParams;
-			SceneTextureInput.Texture = SceneColorRenderTarget.Texture;
-			SceneTextureInput.Sampler = PointClampSampler;
-		}
-
-		// Because we are not using proxy material, but plain global shader, we need to setup Scene textures ourselves.
-		// We don't need to do this per region.
-		check(View.bIsViewInfo);
-		FSceneTextureShaderParameters SceneTextures = CreateSceneTextureShaderParameters(GraphBuilder, ((const FViewInfo&)View).GetSceneTexturesChecked(), View.GetFeatureLevel(), ESceneTextureSetupMode::All);
-
-		for (auto It = WorldSubsystem->Regions.CreateConstIterator(); It; ++It)
-		{
-			AColorCorrectRegion* Region = *It;
-
-			/* If Region is pending for kill, invisible or disabled we don't need to render it.
-			*	If Region's Primitive is not visible in the current view's scene then we don't need to render it either.
-			*	We are checking if the region belongs to the same world as the view. 
-			* Alternative is to get all component ids from ViewFamily.Scene and compare with actor's 
-			*	Region->ActiveMeshComponent->ComponentId ComponentIdSearchTable.Contains(RegionPrimitiveComponentId)
-			*/
-			if (!IsValid(Region) || 
-				Region->IsActorBeingDestroyed() ||
-				!Region->Enabled || 
-				Region->IsHidden() ||
-#if WITH_EDITOR
-				Region->IsHiddenEd() ||
-#endif
-				Region->GetWorld() != ViewFamily.Scene->GetWorld() ||
-				// Display Cluster uses HiddenPrimitives to hide Primitive components from view.
-				View.HiddenPrimitives.Contains(Region->GetFirstPrimitiveId()))
-			{
-				continue;
-			}
-
-			FVector BoxCenter, BoxExtents;
-			Region->GetBounds(BoxCenter, BoxExtents);
-
-			// If bounding box is zero, then we don't need to do anything.
-			if (BoxExtents.IsNearlyZero())
-			{
-				continue;
-			}
-
-			FIntRect Viewport;
-
-			float MaxDepth = -BIG_NUMBER;
-			float MinDepth = BIG_NUMBER;
-
-			if (Region->Invert)
-			{
-				// In case of Region inversion we would to render the entire screen
-				Viewport = PrimaryViewRect;
-			}
-			else
-			{
-				GetPixelSpaceBoundingRect(View, BoxCenter, BoxExtents, Viewport, MaxDepth, MinDepth);
-
-				// Check if CCR is too small to be rendered (less than one pixel on the screen).
-				if (Viewport.Width() == 0 || Viewport.Height() == 0)
-				{
-					continue;
-				}
-
-				// This is to handle corner cases when user has a very long disproportionate region and gets either
-				// within bounds or close to the center.
-				float MaxBoxExtent = FMath::Abs(BoxExtents.GetMax());
-				if (MaxDepth >= 0 && MinDepth < 0)
-				{
-					UpdateMinMaxWithFrustrumAABBIntersection(View, BoxCenter, BoxExtents, Viewport, MaxDepth);
-				}
-
-				FIntRect ConstrainedViewRect = View.UnscaledViewRect;
-
-				// We need to make sure that Bounding Rectangle is offset by the position of the View's Viewport.
-				Viewport.Min -= ConstrainedViewRect.Min;
-
-				Viewport = Viewport.Scale(ScreenPercentage);
-
-				// Culling all regions that are not within the screen bounds.
-				if ((Viewport.Min.X >= PrimaryViewRect.Width() ||
-					Viewport.Min.Y >= PrimaryViewRect.Height() ||
-					Viewport.Max.X <= 0.0f ||
-					Viewport.Max.Y <= 0.0f ||
-					MaxDepth < 0.0f))
-				{
-					continue;
-				}
-				// Clipping is required because as we get closer to the bounding box the bounds
-				// May extend beyond Allowed render target size.
-				Viewport.Clip(PrimaryViewRect);
-			}
-			
-
-			bool bIsAdvanced = false;
-
-			const FVector4 One(1., 1., 1., 1.);
-			const FVector4 Zero(0., 0., 0., 0.);
-			TArray<FColorGradePerRangeSettings*> AdvancedSettings{ &Region->ColorGradingSettings.Shadows,
-																	&Region->ColorGradingSettings.Midtones,
-																	&Region->ColorGradingSettings.Highlights };
-
-			// Check if any of the regions are advanced.
-			for (auto SettingsIt = AdvancedSettings.CreateConstIterator(); SettingsIt; ++SettingsIt)
-			{
-				const FColorGradePerRangeSettings* ColorGradingSettings = *SettingsIt;
-				if (!ColorGradingSettings->Saturation.Equals(One, SMALL_NUMBER) ||
-					!ColorGradingSettings->Contrast.Equals(One, SMALL_NUMBER) ||
-					!ColorGradingSettings->Gamma.Equals(One, SMALL_NUMBER) ||
-					!ColorGradingSettings->Gain.Equals(One, SMALL_NUMBER) ||
-					!ColorGradingSettings->Offset.Equals(Zero, SMALL_NUMBER))
-				{
-					bIsAdvanced = true;
-					break;
-				}
-			}
-
-			const FScreenPassTextureViewport RegionViewport(SceneColorRenderTarget.Texture, Viewport);
-
-			FCCRShaderInputParameters* PostProcessMaterialParameters = GraphBuilder.AllocParameters<FCCRShaderInputParameters>();
-			PostProcessMaterialParameters->RenderTargets[0] = BackBufferRenderTarget.GetRenderTargetBinding();
-
-			PostProcessMaterialParameters->PostProcessOutput = SceneTextureViewportParams;
-			PostProcessMaterialParameters->PostProcessInput[0] = SceneTextureInput;
-			PostProcessMaterialParameters->SceneTextures = SceneTextures;
-			PostProcessMaterialParameters->View = View.ViewUniformBuffer;
-
-			TShaderMapRef<FColorCorrectRegionMaterialVS> VertexShader(GlobalShaderMap);
-			const float DefaultTemperature = 6500;
-
-			// If temperature is default we don't want to do the calculations.
-			FColorCorrectRegionMaterialPS::ETemperatureType TemperatureType = FMath::IsNearlyEqual(Region->Temperature, DefaultTemperature) 
-				? FColorCorrectRegionMaterialPS::ETemperatureType::Disabled 
-				: static_cast<FColorCorrectRegionMaterialPS::ETemperatureType>(Region->TemperatureType);
-			TShaderMapRef<FColorCorrectRegionMaterialPS> PixelShader = GetRegionShader(GlobalShaderMap, Region->Type, TemperatureType, bIsAdvanced, bSampleOpacityFromGbuffer);
-
-			ClearUnusedGraphResources(VertexShader, PixelShader, PostProcessMaterialParameters);
-
-			FCCRRegionDataInputParameter RegionData;
-			FCCRColorCorrectParameter CCBase;
-			FCCRColorCorrectShadowsParameter CCShadows;
-			FCCRColorCorrectMidtonesParameter CCMidtones;
-			FCCRColorCorrectHighlightsParameter CCHighlights;
-
-
-			// Setting constant buffer data to be passed to the shader.
-			{
-				RegionData.Rotate = FMath::DegreesToRadians<FVector3f>((FVector3f)Region->GetActorRotation().Euler());	// LWC_TODO: Precision Loss
-				RegionData.Translate = (FVector3f)Region->GetActorLocation();	// LWC_TODO: Precision Loss
-
-				const float ScaleMultiplier = 50.;
-				// Pre multiplied scale. 
-				RegionData.Scale = (FVector3f)Region->GetActorScale() * ScaleMultiplier;
-
-				RegionData.WhiteTemp = Region->Temperature;
-				// Inner could be larger than outer, in which case we need to make sure these are swapped.
-				RegionData.Inner = FMath::Min<float>(Region->Outer, Region->Inner);
-				RegionData.Outer = FMath::Max<float>(Region->Outer, Region->Inner);
-				RegionData.Falloff = Region->Falloff;
-				RegionData.Intensity = Region->Intensity;
-				RegionData.ExcludeStencil = Region->ExcludeStencil;
-				RegionData.Invert = Region->Invert;
-
-				CCBase.ColorSaturation = (FVector4f)Region->ColorGradingSettings.Global.Saturation;
-				CCBase.ColorContrast = (FVector4f)Region->ColorGradingSettings.Global.Contrast;
-				CCBase.ColorGamma = (FVector4f)Region->ColorGradingSettings.Global.Gamma;
-				CCBase.ColorGain = (FVector4f)Region->ColorGradingSettings.Global.Gain;
-				CCBase.ColorOffset = (FVector4f)Region->ColorGradingSettings.Global.Offset;
-
-				// Set advanced 
-				if (bIsAdvanced)
-				{
-					const float GammaMin = 0.02;
-					const float GammaMax = 10.;
-					//clamp(ExternalExpressions.ColorGammaHighlights, 0.02, 10.)
-					CCShadows.ColorSaturation = (FVector4f)Region->ColorGradingSettings.Shadows.Saturation;
-					CCShadows.ColorContrast = (FVector4f)Region->ColorGradingSettings.Shadows.Contrast;
-					CCShadows.ColorGamma = (FVector4f)Clamp(Region->ColorGradingSettings.Shadows.Gamma, GammaMin, GammaMax);
-					CCShadows.ColorGain = (FVector4f)Region->ColorGradingSettings.Shadows.Gain;
-					CCShadows.ColorOffset = (FVector4f)Region->ColorGradingSettings.Shadows.Offset;
-					CCShadows.ShadowMax = Region->ColorGradingSettings.ShadowsMax;
-
-					CCMidtones.ColorSaturation = (FVector4f)Region->ColorGradingSettings.Midtones.Saturation;
-					CCMidtones.ColorContrast = (FVector4f)Region->ColorGradingSettings.Midtones.Contrast;
-					CCMidtones.ColorGamma = (FVector4f)Clamp(Region->ColorGradingSettings.Midtones.Gamma, GammaMin, GammaMax);
-					CCMidtones.ColorGain = (FVector4f)Region->ColorGradingSettings.Midtones.Gain;
-					CCMidtones.ColorOffset = (FVector4f)Region->ColorGradingSettings.Midtones.Offset;
-
-					CCHighlights.ColorSaturation = (FVector4f)Region->ColorGradingSettings.Highlights.Saturation;
-					CCHighlights.ColorContrast = (FVector4f)Region->ColorGradingSettings.Highlights.Contrast;
-					CCHighlights.ColorGamma = (FVector4f)Clamp(Region->ColorGradingSettings.Highlights.Gamma, GammaMin, GammaMax);
-					CCHighlights.ColorGain = (FVector4f)Region->ColorGradingSettings.Highlights.Gain;
-					CCHighlights.ColorOffset = (FVector4f)Region->ColorGradingSettings.Highlights.Offset;
-					CCHighlights.HighlightsMin = Region->ColorGradingSettings.HighlightsMin;
-				}
-			}
-
-#if CLIP_PIXELS_OUTSIDE_AABB
-			// In case this is a second pass we need to clear the viewport in the backbuffer texture.
-			// We don't need to clear the entire texture, just the render viewport.
-			if (BackBufferRenderTarget.LoadAction == ERenderTargetLoadAction::ELoad)
-			{
-				FClearRectPS::FParameters* Parameters = GraphBuilder.AllocParameters<FClearRectPS::FParameters>();
-				TShaderMapRef<FClearRectPS> CopyPixelShader(GlobalShaderMap);
-				TShaderMapRef<FColorCorrectScreenPassVS> ScreenPassVS(GlobalShaderMap);
-				Parameters->RenderTargets[0] = BackBufferRenderTarget.GetRenderTargetBinding();
-
-				GraphBuilder.AddPass(
-					RDG_EVENT_NAME("ColorCorrectRegions_ClearViewport"),
-					Parameters,
-					ERDGPassFlags::Raster,
-					[&View, ScreenPassVS, CopyPixelShader, RegionViewport, Parameters, DefaultBlendState](FRHICommandList& RHICmdList)
-				{
-					DrawScreenPass(
-						RHICmdList,
-						View,
-						RegionViewport,
-						RegionViewport,
-						FScreenPassPipelineState(ScreenPassVS, CopyPixelShader, DefaultBlendState),
-						[&](FRHICommandList&)
-					{
-						SetShaderParameters(RHICmdList, CopyPixelShader, CopyPixelShader.GetPixelShader(), *Parameters);
-					});
-				});
-			}
-#endif
-			// Main region rendering.
-			GraphBuilder.AddPass(
-				RDG_EVENT_NAME("ColorCorrectRegions"),
-				PostProcessMaterialParameters,
-				ERDGPassFlags::Raster,
-				[&View,
-				RegionViewport,
-				VertexShader,
-				PixelShader,
-				DefaultBlendState,
-				DepthStencilState,
-				PostProcessMaterialParameters,
-				RegionData,
-				CCBase,
-				CCShadows,
-				CCMidtones,
-				CCHighlights,
-				bIsAdvanced](FRHICommandList& RHICmdList)
-			{
-				
-				DrawScreenPass(
-					RHICmdList,
-					View,
-					RegionViewport, // Output Viewport
-					RegionViewport, // Input Viewport
-					FScreenPassPipelineState(VertexShader, PixelShader, DefaultBlendState, DepthStencilState),
-					[&](FRHICommandList& RHICmdList)
-				{
-					SetUniformBufferParameterImmediate(RHICmdList, PixelShader.GetPixelShader(), PixelShader->GetUniformBufferParameter<FCCRRegionDataInputParameter>(), RegionData);
-					SetUniformBufferParameterImmediate(RHICmdList, PixelShader.GetPixelShader(), PixelShader->GetUniformBufferParameter<FCCRColorCorrectParameter>(), CCBase);
-					if (bIsAdvanced)
-					{
-						SetUniformBufferParameterImmediate(RHICmdList, PixelShader.GetPixelShader(), PixelShader->GetUniformBufferParameter<FCCRColorCorrectShadowsParameter>(), CCShadows);
-						SetUniformBufferParameterImmediate(RHICmdList, PixelShader.GetPixelShader(), PixelShader->GetUniformBufferParameter<FCCRColorCorrectMidtonesParameter>(), CCMidtones);
-						SetUniformBufferParameterImmediate(RHICmdList, PixelShader.GetPixelShader(), PixelShader->GetUniformBufferParameter<FCCRColorCorrectHighlightsParameter>(), CCHighlights);
-					}
-					VertexShader->SetParameters(RHICmdList,  View);
-					SetShaderParameters(RHICmdList, VertexShader, VertexShader.GetVertexShader(), *PostProcessMaterialParameters);
-
-					PixelShader->SetParameters(RHICmdList, View);
-					SetShaderParameters(RHICmdList, PixelShader, PixelShader.GetPixelShader(), *PostProcessMaterialParameters);
-				});
-
-			});
-
-			// Since we've rendered into the backbuffer already we have to use load flag instead.
-			BackBufferRenderTarget.LoadAction = ERenderTargetLoadAction::ELoad;
-
-			FCopyRectPS::FParameters* Parameters = GraphBuilder.AllocParameters<FCopyRectPS::FParameters>();
-			Parameters->InputTexture = BackBufferRenderTarget.Texture;
-			Parameters->InputSampler = TStaticSamplerState<>::GetRHI();
-			Parameters->RenderTargets[0] = SceneColorRenderTarget.GetRenderTargetBinding();
-
-			TShaderMapRef<FCopyRectPS> CopyPixelShader(GlobalShaderMap);
-			TShaderMapRef<FColorCorrectScreenPassVS> ScreenPassVS(GlobalShaderMap);
-
-#if CLIP_PIXELS_OUTSIDE_AABB
-			// Blending the output from the main step with scene color.
-			// src.rgb*src.a + dest.rgb*(1.-src.a); alpha = src.a*0. + dst.a*1.0
-			FRHIBlendState* CopyBlendState = TStaticBlendState<CW_RGB, BO_Add, BF_SourceAlpha, BF_InverseSourceAlpha, BO_Add, BF_Zero, BF_One>::GetRHI();
-#else	
-			FRHIBlendState* CopyBlendState = DefaultBlendState;
-#endif
-			GraphBuilder.AddPass(
-				RDG_EVENT_NAME("ColorCorrectRegions_CopyViewport"),
-				Parameters,
-				ERDGPassFlags::Raster,
-				[&View, ScreenPassVS, CopyPixelShader, RegionViewport, Parameters, CopyBlendState](FRHICommandList& RHICmdList)
-			{
-				DrawScreenPass(
-					RHICmdList,
-					View,
-					RegionViewport,
-					RegionViewport,
-					FScreenPassPipelineState(ScreenPassVS, CopyPixelShader, CopyBlendState),
-					[&](FRHICommandList&)
-				{
-					SetShaderParameters(RHICmdList, CopyPixelShader, CopyPixelShader.GetPixelShader(), *Parameters);
-				});
-			});
-
-		}
-
-	}
-}
-
+// Copyright Epic Games, Inc. All Rights Reserved.
+
+#include "ColorCorrectRegionsSceneViewExtension.h"
+#include "ColorCorrectRegionsModule.h"
+#include "ColorCorrectRegionsSubsystem.h"
+#include "ColorCorrectRegionsPostProcessMaterial.h"
+#include "CommonRenderResources.h"
+#include "Containers/DynamicRHIResourceArray.h"
+#include "DynamicResolutionState.h"
+#include "Engine/World.h"
+#include "EngineUtils.h"
+#include "PostProcess/PostProcessing.h"
+#include "RHI.h"
+#include "SceneRendering.h"
+#include "ScreenPass.h"
+#include "SceneView.h"
+
+// Set this to 1 to clip pixels outside of bounding box.
+#define CLIP_PIXELS_OUTSIDE_AABB 1
+
+//Set this to 1 to see the clipping region.
+#define ColorCorrectRegions_SHADER_DISPLAY_BOUNDING_RECT 0
+
+DECLARE_GPU_STAT_NAMED(ColorCorrectRegion_Render, TEXT("ColorCorrectRegion_Render"));
+DECLARE_GPU_STAT_NAMED(ColorCorrectRegion_Copy, TEXT("ColorCorrectRegion_Copy"));
+DECLARE_GPU_STAT_NAMED(ColorCorrectRegion_Clear, TEXT("ColorCorrectRegion_Clear"));
+
+namespace
+{
+	FScreenPassTextureViewportParameters GetTextureViewportParameters(const FScreenPassTextureViewport& InViewport)
+	{
+		const FVector2f Extent(InViewport.Extent);
+		const FVector2f ViewportMin(InViewport.Rect.Min.X, InViewport.Rect.Min.Y);
+		const FVector2f ViewportMax(InViewport.Rect.Max.X, InViewport.Rect.Max.Y);
+		const FVector2f ViewportSize = ViewportMax - ViewportMin;
+
+		FScreenPassTextureViewportParameters Parameters;
+
+		if (!InViewport.IsEmpty())
+		{
+			Parameters.Extent = FVector2f(Extent);
+			Parameters.ExtentInverse = FVector2f(1.0f / Extent.X, 1.0f / Extent.Y);
+
+			Parameters.ScreenPosToViewportScale = FVector2f(0.5f, -0.5f) * ViewportSize;	
+			Parameters.ScreenPosToViewportBias = (0.5f * ViewportSize) + ViewportMin;	
+
+			Parameters.ViewportMin = InViewport.Rect.Min;
+			Parameters.ViewportMax = InViewport.Rect.Max;
+
+			Parameters.ViewportSize = ViewportSize;
+			Parameters.ViewportSizeInverse = FVector2f(1.0f / Parameters.ViewportSize.X, 1.0f / Parameters.ViewportSize.Y);
+
+			Parameters.UVViewportMin = ViewportMin * Parameters.ExtentInverse;
+			Parameters.UVViewportMax = ViewportMax * Parameters.ExtentInverse;
+
+			Parameters.UVViewportSize = Parameters.UVViewportMax - Parameters.UVViewportMin;
+			Parameters.UVViewportSizeInverse = FVector2f(1.0f / Parameters.UVViewportSize.X, 1.0f / Parameters.UVViewportSize.Y);
+
+			Parameters.UVViewportBilinearMin = Parameters.UVViewportMin + 0.5f * Parameters.ExtentInverse;
+			Parameters.UVViewportBilinearMax = Parameters.UVViewportMax - 0.5f * Parameters.ExtentInverse;
+		}
+
+		return Parameters;
+	}
+
+
+	void GetPixelSpaceBoundingRect(const FSceneView& InView, const FVector& InBoxCenter, const FVector& InBoxExtents, FIntRect& OutViewport, float& OutMaxDepth, float& OutMinDepth)
+	{
+		OutViewport = FIntRect(INT32_MAX, INT32_MAX, -INT32_MAX, -INT32_MAX);
+		// 8 corners of the bounding box. To be multiplied by box extent and offset by the center.
+		const int NumCorners = 8;
+		const FVector Verts[NumCorners] = {
+			FVector(1, 1, 1),
+			FVector(1, 1,-1),
+			FVector(1,-1, 1),
+			FVector(1,-1,-1),
+			FVector(-1, 1, 1),
+			FVector(-1, 1,-1),
+			FVector(-1,-1, 1),
+			FVector(-1,-1,-1) };
+
+		for (int32 Index = 0; Index < NumCorners; Index++)
+		{
+			// Project bounding box vertecies into screen space.
+			const FVector WorldVert = InBoxCenter + (Verts[Index] * InBoxExtents);
+			FVector2D PixelVert;
+			FVector4 ScreenSpaceCoordinate = InView.WorldToScreen(WorldVert);
+
+			OutMaxDepth = FMath::Max<float>(ScreenSpaceCoordinate.W, OutMaxDepth);
+			OutMinDepth = FMath::Min<float>(ScreenSpaceCoordinate.W, OutMinDepth);
+
+			if (InView.ScreenToPixel(ScreenSpaceCoordinate, PixelVert))
+			{
+				// Update screen-space bounding box with with transformed vert.
+				OutViewport.Min.X = FMath::Min<int32>(OutViewport.Min.X, PixelVert.X);
+				OutViewport.Min.Y = FMath::Min<int32>(OutViewport.Min.Y, PixelVert.Y);
+
+				OutViewport.Max.X = FMath::Max<int32>(OutViewport.Max.X, PixelVert.X);
+				OutViewport.Max.Y = FMath::Max<int32>(OutViewport.Max.Y, PixelVert.Y);
+			}
+		}
+	}
+
+	// Function that calculates all points of intersection between plane and bounding box. Resulting points are unsorted.
+	void CalculatePlaneAABBIntersectionPoints(const FPlane& Plane, const FVector& BoxCenter, const FVector& BoxExtents, TArray<FVector>& OutPoints)
+	{
+		const FVector MaxCorner = BoxCenter + BoxExtents;
+
+		const FVector Verts[3][4] = {
+			{
+				// X Direction
+				FVector(-1, -1, -1),
+				FVector(-1,  1, -1),
+				FVector(-1, -1,  1),
+				FVector(-1,  1,  1),
+			},
+			{
+				// Y Direction
+				FVector(-1, -1, -1),
+				FVector( 1, -1, -1),
+				FVector( 1, -1,  1),
+				FVector(-1, -1,  1),
+			},
+			{
+				// Z Direction
+				FVector(-1, -1, -1),
+				FVector( 1, -1, -1),
+				FVector( 1,  1, -1),
+				FVector(-1,  1, -1),
+			}
+		};
+
+		FVector Intersection;
+		FVector Start;
+		FVector End;
+
+		for (int RunningAxis_Dir = 0; RunningAxis_Dir < 3; RunningAxis_Dir++)
+		{
+			const FVector *CornerLocations = Verts[RunningAxis_Dir];
+			for (int RunningCorner = 0; RunningCorner < 4; RunningCorner++)
+			{
+				Start = BoxCenter + BoxExtents * CornerLocations[RunningCorner];
+				End = FVector(Start.X, Start.Y, Start.Z);
+				End[RunningAxis_Dir] = MaxCorner[RunningAxis_Dir];
+				if (FMath::SegmentPlaneIntersection(Start, End, Plane, Intersection))
+				{
+					OutPoints.Add(Intersection);
+				}
+			}
+		}
+	}
+
+	// Takes in an existing viewport and updates it with an intersection bounding rectangle.
+	void UpdateMinMaxWithFrustrumAABBIntersection(const FSceneView& InView, const FVector& InBoxCenter, const FVector& InBoxExtents, FIntRect& OutViewportToUpdate, float& OutMaxDepthToUpdate)
+	{
+		TArray<FVector> Points;
+		Points.Reserve(6);
+		static bool bNotifiedOfClippingPlaneError = false;
+
+		if (InView.bHasNearClippingPlane)
+		{
+			CalculatePlaneAABBIntersectionPoints(InView.NearClippingPlane, InBoxCenter, InBoxExtents, Points);
+		}
+		// Previously last plane was near clipping plane.
+		else if (InView.ViewFrustum.Planes.Num() == 5)
+		{
+			CalculatePlaneAABBIntersectionPoints(InView.ViewFrustum.Planes[4], InBoxCenter, InBoxExtents, Points);
+		}
+		else if (!bNotifiedOfClippingPlaneError)
+		{
+			bNotifiedOfClippingPlaneError = true;
+			UE_LOG(ColorCorrectRegions, Error, TEXT("Couldn't find a correct near clipping plane in View Frustrum"));
+		}
+
+		if (Points.Num() == 0)
+		{
+			return;
+		}
+
+		for (FVector Point : Points)
+		{
+			// Project bounding box vertecies into screen space.
+			FVector4 ScreenSpaceCoordinate = InView.WorldToScreen(Point);
+			FVector4 ScreenSpaceCoordinateScaled = ScreenSpaceCoordinate * 1.0 / ScreenSpaceCoordinate.W;
+
+			OutMaxDepthToUpdate = FMath::Max<float>(ScreenSpaceCoordinate.W, OutMaxDepthToUpdate);
+			FVector2D PixelVert;
+
+			if (InView.ScreenToPixel(ScreenSpaceCoordinate, PixelVert))
+			{
+				// Update screen-space bounding box with with transformed vert.
+				OutViewportToUpdate.Min.X = FMath::Min<int32>(OutViewportToUpdate.Min.X, PixelVert.X);
+				OutViewportToUpdate.Min.Y = FMath::Min<int32>(OutViewportToUpdate.Min.Y, PixelVert.Y);
+
+				OutViewportToUpdate.Max.X = FMath::Max<int32>(OutViewportToUpdate.Max.X, PixelVert.X);
+				OutViewportToUpdate.Max.Y = FMath::Max<int32>(OutViewportToUpdate.Max.Y, PixelVert.Y);
+			}
+		}
+	}
+
+	bool ViewSupportsRegions(const FSceneView& View)
+	{
+		return View.Family->EngineShowFlags.PostProcessing &&
+				View.Family->EngineShowFlags.PostProcessMaterial;
+	}
+
+	// A helper function for getting the right shader.
+	TShaderMapRef<FColorCorrectRegionMaterialPS> GetRegionShader(const FGlobalShaderMap* GlobalShaderMap, EColorCorrectRegionsType RegionType, FColorCorrectRegionMaterialPS::ETemperatureType TemperatureType, bool bIsAdvanced)
+	{
+		FColorCorrectRegionMaterialPS::FPermutationDomain PermutationVector;
+		PermutationVector.Set<FColorCorrectRegionMaterialPS::FAdvancedShader>(bIsAdvanced);
+		PermutationVector.Set<FColorCorrectRegionMaterialPS::FShaderType>(static_cast<EColorCorrectRegionsType>(FMath::Min(static_cast<int32>(RegionType), static_cast<int32>(EColorCorrectRegionsType::MAX) - 1)));
+		PermutationVector.Set<FColorCorrectRegionMaterialPS::FTemperatureType>(TemperatureType);
+
+		return TShaderMapRef<FColorCorrectRegionMaterialPS>(GlobalShaderMap, PermutationVector);
+		;
+	}
+
+	FVector4 Clamp(const FVector4 & VectorToClamp, float Min, float Max)
+	{
+		return FVector4(FMath::Clamp(VectorToClamp.X, Min, Max),
+						FMath::Clamp(VectorToClamp.Y, Min, Max),
+						FMath::Clamp(VectorToClamp.Z, Min, Max),
+						FMath::Clamp(VectorToClamp.W, Min, Max));
+	}
+
+
+	bool RenderRegion
+	(FRDGBuilder& GraphBuilder
+		, const FSceneView& View
+		, const FPostProcessingInputs& Inputs
+		, const FSceneViewFamily& ViewFamily
+		, AColorCorrectRegion* Region
+		, const FIntRect& PrimaryViewRect
+		, const FScreenPassRenderTarget& SceneColorRenderTarget
+		, const float ScreenPercentage
+		, FScreenPassRenderTarget& BackBufferRenderTarget
+		, const FScreenPassTextureViewportParameters& SceneTextureViewportParams
+		, const FScreenPassTextureInput& SceneTextureInput
+		, const FSceneTextureShaderParameters& SceneTextures
+		, FGlobalShaderMap* GlobalShaderMap
+		, FRHIBlendState* DefaultBlendState)
+	{
+		FRHIDepthStencilState* DepthStencilState = FScreenPassPipelineState::FDefaultDepthStencilState::GetRHI();
+
+		/* If Region is pending for kill, invisible or disabled we don't need to render it.
+		*	If Region's Primitive is not visible in the current view's scene then we don't need to render it either.
+		*	We are checking if the region belongs to the same world as the view.
+		* Alternative is to get all component ids from ViewFamily.Scene and compare with actor's
+		*	Region->ActiveMeshComponent->ComponentId ComponentIdSearchTable.Contains(RegionPrimitiveComponentId)
+		*/
+		if (!IsValid(Region) ||
+			Region->IsActorBeingDestroyed() ||
+			!Region->Enabled ||
+			Region->IsHidden() ||
+#if WITH_EDITOR
+			Region->IsHiddenEd() ||
+#endif
+			Region->GetWorld() != ViewFamily.Scene->GetWorld())
+		{
+			return false;
+		}
+
+		FVector BoxCenter, BoxExtents;
+		Region->GetBounds(BoxCenter, BoxExtents);
+
+		// If bounding box is zero, then we don't need to do anything.
+		if (BoxExtents.IsNearlyZero())
+		{
+			return false;
+		}
+
+		FIntRect Viewport;
+
+		float MaxDepth = -BIG_NUMBER;
+		float MinDepth = BIG_NUMBER;
+
+		if (Region->Invert)
+		{
+			// In case of Region inversion we would to render the entire screen
+			Viewport = PrimaryViewRect;
+		}
+		else
+		{
+			GetPixelSpaceBoundingRect(View, BoxCenter, BoxExtents, Viewport, MaxDepth, MinDepth);
+
+			// Check if CCR is too small to be rendered (less than one pixel on the screen).
+			if (Viewport.Width() == 0 || Viewport.Height() == 0)
+			{
+				return false;
+			}
+
+			// This is to handle corner cases when user has a very long disproportionate region and gets either
+			// within bounds or close to the center.
+			float MaxBoxExtent = FMath::Abs(BoxExtents.GetMax());
+			if (MaxDepth >= 0 && MinDepth < 0)
+			{
+				UpdateMinMaxWithFrustrumAABBIntersection(View, BoxCenter, BoxExtents, Viewport, MaxDepth);
+			}
+
+			FIntRect ConstrainedViewRect = View.UnscaledViewRect;
+
+			// We need to make sure that Bounding Rectangle is offset by the position of the View's Viewport.
+			Viewport.Min -= ConstrainedViewRect.Min;
+
+			Viewport = Viewport.Scale(ScreenPercentage);
+
+			// Culling all regions that are not within the screen bounds.
+			if ((Viewport.Min.X >= PrimaryViewRect.Width() ||
+				Viewport.Min.Y >= PrimaryViewRect.Height() ||
+				Viewport.Max.X <= 0.0f ||
+				Viewport.Max.Y <= 0.0f ||
+				MaxDepth < 0.0f))
+			{
+				return false;
+			}
+			// Clipping is required because as we get closer to the bounding box the bounds
+			// May extend beyond Allowed render target size.
+			Viewport.Clip(PrimaryViewRect);
+		}
+
+
+		bool bIsAdvanced = false;
+
+		const FVector4 One(1., 1., 1., 1.);
+		const FVector4 Zero(0., 0., 0., 0.);
+		TArray<FColorGradePerRangeSettings*> AdvancedSettings{ &Region->ColorGradingSettings.Shadows,
+																&Region->ColorGradingSettings.Midtones,
+																&Region->ColorGradingSettings.Highlights };
+
+		// Check if any of the regions are advanced.
+		for (auto SettingsIt = AdvancedSettings.CreateConstIterator(); SettingsIt; ++SettingsIt)
+		{
+			const FColorGradePerRangeSettings* ColorGradingSettings = *SettingsIt;
+			if (!ColorGradingSettings->Saturation.Equals(One, SMALL_NUMBER) ||
+				!ColorGradingSettings->Contrast.Equals(One, SMALL_NUMBER) ||
+				!ColorGradingSettings->Gamma.Equals(One, SMALL_NUMBER) ||
+				!ColorGradingSettings->Gain.Equals(One, SMALL_NUMBER) ||
+				!ColorGradingSettings->Offset.Equals(Zero, SMALL_NUMBER))
+			{
+				bIsAdvanced = true;
+				break;
+			}
+		}
+
+		const FScreenPassTextureViewport RegionViewport(SceneColorRenderTarget.Texture, Viewport);
+
+		FCCRShaderInputParameters* PostProcessMaterialParameters = GraphBuilder.AllocParameters<FCCRShaderInputParameters>();
+		PostProcessMaterialParameters->RenderTargets[0] = BackBufferRenderTarget.GetRenderTargetBinding();
+
+		PostProcessMaterialParameters->PostProcessOutput = SceneTextureViewportParams;
+		PostProcessMaterialParameters->PostProcessInput[0] = SceneTextureInput;
+		PostProcessMaterialParameters->SceneTextures = SceneTextures;
+		PostProcessMaterialParameters->View = View.ViewUniformBuffer;
+
+		TShaderMapRef<FColorCorrectRegionMaterialVS> VertexShader(GlobalShaderMap);
+		const float DefaultTemperature = 6500;
+
+		// If temperature is default we don't want to do the calculations.
+		FColorCorrectRegionMaterialPS::ETemperatureType TemperatureType = FMath::IsNearlyEqual(Region->Temperature, DefaultTemperature)
+			? FColorCorrectRegionMaterialPS::ETemperatureType::Disabled
+			: static_cast<FColorCorrectRegionMaterialPS::ETemperatureType>(Region->TemperatureType);
+		TShaderMapRef<FColorCorrectRegionMaterialPS> PixelShader = GetRegionShader(GlobalShaderMap, Region->Type, TemperatureType, bIsAdvanced);
+
+		ClearUnusedGraphResources(VertexShader, PixelShader, PostProcessMaterialParameters);
+
+		FCCRRegionDataInputParameter RegionData;
+		FCCRColorCorrectParameter CCBase;
+		FCCRColorCorrectShadowsParameter CCShadows;
+		FCCRColorCorrectMidtonesParameter CCMidtones;
+		FCCRColorCorrectHighlightsParameter CCHighlights;
+
+
+		// Setting constant buffer data to be passed to the shader.
+		{
+			RegionData.Rotate = FMath::DegreesToRadians<FVector3f>((FVector3f)Region->GetActorRotation().Euler());
+			RegionData.Translate = (FVector3f)Region->GetActorLocation();
+
+			const float ScaleMultiplier = 50.;
+			// Pre multiplied scale. 
+			RegionData.Scale = (FVector3f)Region->GetActorScale() * ScaleMultiplier;
+
+			RegionData.WhiteTemp = Region->Temperature;
+			// Inner could be larger than outer, in which case we need to make sure these are swapped.
+			RegionData.Inner = FMath::Min<float>(Region->Outer, Region->Inner);
+			RegionData.Outer = FMath::Max<float>(Region->Outer, Region->Inner);
+			if (RegionData.Inner == RegionData.Outer)
+			{
+				RegionData.Inner -= 0.0001;
+			}
+			RegionData.Falloff = Region->Falloff;
+			RegionData.Intensity = Region->Intensity;
+			RegionData.ExcludeStencil = Region->ExcludeStencil;
+			RegionData.Invert = Region->Invert;
+
+			CCBase.ColorSaturation = (FVector4f)Region->ColorGradingSettings.Global.Saturation;
+			CCBase.ColorContrast = (FVector4f)Region->ColorGradingSettings.Global.Contrast;
+			CCBase.ColorGamma = (FVector4f)Region->ColorGradingSettings.Global.Gamma;
+			CCBase.ColorGain = (FVector4f)Region->ColorGradingSettings.Global.Gain;
+			CCBase.ColorOffset = (FVector4f)Region->ColorGradingSettings.Global.Offset;
+
+			// Set advanced 
+			if (bIsAdvanced)
+			{
+				const float GammaMin = 0.02;
+				const float GammaMax = 10.;
+				//clamp(ExternalExpressions.ColorGammaHighlights, 0.02, 10.)
+				CCShadows.ColorSaturation = (FVector4f)Region->ColorGradingSettings.Shadows.Saturation;
+				CCShadows.ColorContrast = (FVector4f)Region->ColorGradingSettings.Shadows.Contrast;
+				CCShadows.ColorGamma = (FVector4f)Clamp(Region->ColorGradingSettings.Shadows.Gamma, GammaMin, GammaMax);
+				CCShadows.ColorGain = (FVector4f)Region->ColorGradingSettings.Shadows.Gain;
+				CCShadows.ColorOffset = (FVector4f)Region->ColorGradingSettings.Shadows.Offset;
+				CCShadows.ShadowMax = Region->ColorGradingSettings.ShadowsMax;
+
+				CCMidtones.ColorSaturation = (FVector4f)Region->ColorGradingSettings.Midtones.Saturation;
+				CCMidtones.ColorContrast = (FVector4f)Region->ColorGradingSettings.Midtones.Contrast;
+				CCMidtones.ColorGamma = (FVector4f)Clamp(Region->ColorGradingSettings.Midtones.Gamma, GammaMin, GammaMax);
+				CCMidtones.ColorGain = (FVector4f)Region->ColorGradingSettings.Midtones.Gain;
+				CCMidtones.ColorOffset = (FVector4f)Region->ColorGradingSettings.Midtones.Offset;
+
+				CCHighlights.ColorSaturation = (FVector4f)Region->ColorGradingSettings.Highlights.Saturation;
+				CCHighlights.ColorContrast = (FVector4f)Region->ColorGradingSettings.Highlights.Contrast;
+				CCHighlights.ColorGamma = (FVector4f)Clamp(Region->ColorGradingSettings.Highlights.Gamma, GammaMin, GammaMax);
+				CCHighlights.ColorGain = (FVector4f)Region->ColorGradingSettings.Highlights.Gain;
+				CCHighlights.ColorOffset = (FVector4f)Region->ColorGradingSettings.Highlights.Offset;
+				CCHighlights.HighlightsMin = Region->ColorGradingSettings.HighlightsMin;
+			}
+		}
+
+#if CLIP_PIXELS_OUTSIDE_AABB
+		// In case this is a second pass we need to clear the viewport in the backbuffer texture.
+		// We don't need to clear the entire texture, just the render viewport.
+		if (BackBufferRenderTarget.LoadAction == ERenderTargetLoadAction::ELoad)
+		{
+			FClearRectPS::FParameters* Parameters = GraphBuilder.AllocParameters<FClearRectPS::FParameters>();
+			TShaderMapRef<FClearRectPS> CopyPixelShader(GlobalShaderMap);
+			TShaderMapRef<FColorCorrectScreenPassVS> ScreenPassVS(GlobalShaderMap);
+			Parameters->RenderTargets[0] = BackBufferRenderTarget.GetRenderTargetBinding();
+
+			GraphBuilder.AddPass(
+				RDG_EVENT_NAME("ColorCorrectRegions_ClearViewport"),
+				Parameters,
+				ERDGPassFlags::Raster,
+				[&View, ScreenPassVS, CopyPixelShader, RegionViewport, Parameters, DefaultBlendState](FRHICommandList& RHICmdList)
+				{
+					SCOPED_GPU_STAT(RHICmdList, ColorCorrectRegion_Clear);
+					DrawScreenPass(
+						RHICmdList,
+						static_cast<const FViewInfo&>(View),
+						RegionViewport,
+						RegionViewport,
+						FScreenPassPipelineState(ScreenPassVS, CopyPixelShader, DefaultBlendState),
+						EScreenPassDrawFlags::None,
+						[&](FRHICommandList&)
+						{
+							SetShaderParameters(RHICmdList, CopyPixelShader, CopyPixelShader.GetPixelShader(), *Parameters);
+						});
+				});
+		}
+#endif
+		// Main region rendering.
+		GraphBuilder.AddPass(
+			RDG_EVENT_NAME("ColorCorrectRegions"),
+			PostProcessMaterialParameters,
+			ERDGPassFlags::Raster,
+			[&View,
+			RegionViewport,
+			VertexShader,
+			PixelShader,
+			DefaultBlendState,
+			DepthStencilState,
+			PostProcessMaterialParameters,
+			RegionData,
+			CCBase,
+			CCShadows,
+			CCMidtones,
+			CCHighlights,
+			bIsAdvanced](FRHICommandList& RHICmdList)
+			{
+				SCOPED_GPU_STAT(RHICmdList, ColorCorrectRegion_Render);
+
+				DrawScreenPass(
+					RHICmdList,
+					static_cast<const FViewInfo&>(View),
+					RegionViewport, // Output Viewport
+					RegionViewport, // Input Viewport
+					FScreenPassPipelineState(VertexShader, PixelShader, DefaultBlendState, DepthStencilState),
+					EScreenPassDrawFlags::None,
+					[&](FRHICommandList& RHICmdList)
+					{
+						SetUniformBufferParameterImmediate(RHICmdList, PixelShader.GetPixelShader(), PixelShader->GetUniformBufferParameter<FCCRRegionDataInputParameter>(), RegionData);
+						SetUniformBufferParameterImmediate(RHICmdList, PixelShader.GetPixelShader(), PixelShader->GetUniformBufferParameter<FCCRColorCorrectParameter>(), CCBase);
+						if (bIsAdvanced)
+						{
+							SetUniformBufferParameterImmediate(RHICmdList, PixelShader.GetPixelShader(), PixelShader->GetUniformBufferParameter<FCCRColorCorrectShadowsParameter>(), CCShadows);
+							SetUniformBufferParameterImmediate(RHICmdList, PixelShader.GetPixelShader(), PixelShader->GetUniformBufferParameter<FCCRColorCorrectMidtonesParameter>(), CCMidtones);
+							SetUniformBufferParameterImmediate(RHICmdList, PixelShader.GetPixelShader(), PixelShader->GetUniformBufferParameter<FCCRColorCorrectHighlightsParameter>(), CCHighlights);
+						}
+						VertexShader->SetParameters(RHICmdList, View);
+						SetShaderParameters(RHICmdList, VertexShader, VertexShader.GetVertexShader(), *PostProcessMaterialParameters);
+
+						PixelShader->SetParameters(RHICmdList, View);
+						SetShaderParameters(RHICmdList, PixelShader, PixelShader.GetPixelShader(), *PostProcessMaterialParameters);
+					});
+
+			});
+
+		// Since we've rendered into the backbuffer already we have to use load flag instead.
+		BackBufferRenderTarget.LoadAction = ERenderTargetLoadAction::ELoad;
+
+		FCopyRectPS::FParameters* Parameters = GraphBuilder.AllocParameters<FCopyRectPS::FParameters>();
+		Parameters->InputTexture = BackBufferRenderTarget.Texture;
+		Parameters->InputSampler = TStaticSamplerState<>::GetRHI();
+		Parameters->RenderTargets[0] = SceneColorRenderTarget.GetRenderTargetBinding();
+
+		TShaderMapRef<FCopyRectPS> CopyPixelShader(GlobalShaderMap);
+		TShaderMapRef<FColorCorrectScreenPassVS> ScreenPassVS(GlobalShaderMap);
+
+#if CLIP_PIXELS_OUTSIDE_AABB
+		// Blending the output from the main step with scene color.
+		// src.rgb*src.a + dest.rgb*(1.-src.a); alpha = src.a*0. + dst.a*1.0
+		FRHIBlendState* CopyBlendState = TStaticBlendState<CW_RGB, BO_Add, BF_SourceAlpha, BF_InverseSourceAlpha, BO_Add, BF_Zero, BF_One>::GetRHI();
+#else	
+		FRHIBlendState* CopyBlendState = DefaultBlendState;
+#endif
+		GraphBuilder.AddPass(
+			RDG_EVENT_NAME("ColorCorrectRegions_CopyViewport"),
+			Parameters,
+			ERDGPassFlags::Raster,
+			[&View, ScreenPassVS, CopyPixelShader, RegionViewport, Parameters, CopyBlendState](FRHICommandList& RHICmdList)
+			{
+				SCOPED_GPU_STAT(RHICmdList, ColorCorrectRegion_Copy);
+				DrawScreenPass(
+					RHICmdList,
+					static_cast<const FViewInfo&>(View),
+					RegionViewport,
+					RegionViewport,
+					FScreenPassPipelineState(ScreenPassVS, CopyPixelShader, CopyBlendState),
+					EScreenPassDrawFlags::None,
+					[&](FRHICommandList&)
+					{
+						SetShaderParameters(RHICmdList, CopyPixelShader, CopyPixelShader.GetPixelShader(), *Parameters);
+					});
+			});
+
+		return true;
+
+	}
+
+}
+
+FColorCorrectRegionsSceneViewExtension::FColorCorrectRegionsSceneViewExtension(const FAutoRegister& AutoRegister, UColorCorrectRegionsSubsystem* InWorldSubsystem) :
+	FSceneViewExtensionBase(AutoRegister), WorldSubsystem(InWorldSubsystem)
+{
+}
+
+void FColorCorrectRegionsSceneViewExtension::PrePostProcessPass_RenderThread(FRDGBuilder& GraphBuilder, const FSceneView& View, const FPostProcessingInputs& Inputs)
+{
+	// Necessary for when an actor is added or removed from the scene. Also when priority is changed.
+	FScopeLock RegionScopeLock(&WorldSubsystem->RegionAccessCriticalSection);
+
+	if ((WorldSubsystem->RegionsPriorityBased.Num() == 0 && WorldSubsystem->RegionsDistanceBased.Num() == 0)|| !ViewSupportsRegions(View))
+	{
+		return;
+	}
+
+	Inputs.Validate();
+
+	const FSceneViewFamily& ViewFamily = *View.Family;
+
+	DynamicRenderScaling::TMap<float> UpperBounds = ViewFamily.GetScreenPercentageInterface()->GetResolutionFractionsUpperBound();
+	const auto FeatureLevel = View.GetFeatureLevel();
+	const float ScreenPercentage = UpperBounds[GDynamicPrimaryResolutionFraction] * ViewFamily.SecondaryViewFraction;
+	
+	// We need to make sure to take Windows and Scene scale into account.
+
+	checkSlow(View.bIsViewInfo); // can't do dynamic_cast because FViewInfo doesn't have any virtual functions.
+	const FIntRect PrimaryViewRect = static_cast<const FViewInfo&>(View).ViewRect;
+
+	FScreenPassTexture SceneColor((*Inputs.SceneTextures)->SceneColorTexture, PrimaryViewRect);
+
+
+	if (!SceneColor.IsValid())
+	{
+		return;
+	}
+
+	{
+		// Getting material data for the current view.
+		FGlobalShaderMap* GlobalShaderMap = GetGlobalShaderMap(GMaxRHIFeatureLevel);
+
+		// Reusing the same output description for our back buffer as SceneColor
+		FRDGTextureDesc ColorCorrectRegionsOutputDesc = SceneColor.Texture->Desc;
+
+		ColorCorrectRegionsOutputDesc.Format = PF_FloatRGBA;
+		FLinearColor ClearColor(0., 0., 0., 0.);
+		ColorCorrectRegionsOutputDesc.ClearValue = FClearValueBinding(ClearColor);
+
+		FRDGTexture* BackBufferRenderTargetTexture = GraphBuilder.CreateTexture(ColorCorrectRegionsOutputDesc, TEXT("BackBufferRenderTargetTexture"));
+		FScreenPassRenderTarget BackBufferRenderTarget = FScreenPassRenderTarget(BackBufferRenderTargetTexture, SceneColor.ViewRect, ERenderTargetLoadAction::EClear);
+		FScreenPassRenderTarget SceneColorRenderTarget(SceneColor, ERenderTargetLoadAction::ELoad);
+		const FScreenPassTextureViewport SceneColorTextureViewport(SceneColor);
+
+		FRHIBlendState* DefaultBlendState = FScreenPassPipelineState::FDefaultBlendState::GetRHI();
+
+		RDG_EVENT_SCOPE(GraphBuilder, "Color Correct Regions %dx%d", SceneColorTextureViewport.Rect.Width(), SceneColorTextureViewport.Rect.Height());
+
+		FRHISamplerState* PointClampSampler = TStaticSamplerState<SF_Point, AM_Clamp, AM_Clamp, AM_Clamp>::GetRHI();
+		const FScreenPassTextureViewportParameters SceneTextureViewportParams = GetTextureViewportParameters(SceneColorTextureViewport);
+		FScreenPassTextureInput SceneTextureInput;
+		{
+			SceneTextureInput.Viewport = SceneTextureViewportParams;
+			SceneTextureInput.Texture = SceneColorRenderTarget.Texture;
+			SceneTextureInput.Sampler = PointClampSampler;
+		}
+
+		// Because we are not using proxy material, but plain global shader, we need to setup Scene textures ourselves.
+		// We don't need to do this per region.
+		check(View.bIsViewInfo);
+		FSceneTextureShaderParameters SceneTextures = CreateSceneTextureShaderParameters(GraphBuilder, ((const FViewInfo&)View).GetSceneTexturesChecked(), View.GetFeatureLevel(), ESceneTextureSetupMode::All);
+
+
+		for (auto It = WorldSubsystem->RegionsPriorityBased.CreateConstIterator(); It; ++It)
+		{
+			AColorCorrectRegion* Region = *It;
+			RenderRegion(GraphBuilder
+				, View
+				, Inputs
+				, ViewFamily
+				, Region
+				, PrimaryViewRect
+				, SceneColorRenderTarget
+				, ScreenPercentage
+				, BackBufferRenderTarget
+				, SceneTextureViewportParams
+				, SceneTextureInput
+				, SceneTextures
+				, GlobalShaderMap
+				, DefaultBlendState);
+		}
+		WorldSubsystem->SortRegionsByDistance(View.ViewLocation);
+		for (auto It = WorldSubsystem->RegionsDistanceBased.CreateConstIterator(); It; ++It)
+		{
+			AColorCorrectRegion* Region = *It;
+			RenderRegion(GraphBuilder
+				, View
+				, Inputs
+				, ViewFamily
+				, Region
+				, PrimaryViewRect
+				, SceneColorRenderTarget
+				, ScreenPercentage
+				, BackBufferRenderTarget
+				, SceneTextureViewportParams
+				, SceneTextureInput
+				, SceneTextures
+				, GlobalShaderMap
+				, DefaultBlendState);
+		}
+	}
+}
diff --git a/Engine/Plugins/Experimental/ColorCorrectRegions/Source/ColorCorrectRegions/Public/ColorCorrectRegion.h b/Engine/Plugins/Experimental/ColorCorrectRegions/Source/ColorCorrectRegions/Public/ColorCorrectRegion.h
--- a/Engine/Plugins/Experimental/ColorCorrectRegions/Source/ColorCorrectRegions/Public/ColorCorrectRegion.h
+++ b/Engine/Plugins/Experimental/ColorCorrectRegions/Source/ColorCorrectRegions/Public/ColorCorrectRegion.h
@@ -20,6 +20,7 @@ enum class EColorCorrectRegionsType : uint8
 	Box			UMETA(DisplayName = "Box"),
 	Cylinder	UMETA(DisplayName = "Cylinder"),
 	Cone		UMETA(DisplayName = "Cone"),
+	Plane		UMETA(DisplayName = "Plane (Window CCR)"),
 	MAX
 };
 
@@ -50,7 +51,11 @@ public:
 	UPROPERTY(BlueprintReadWrite, EditAnywhere, Category="Color Correction")
 	EColorCorrectRegionsType Type;
 
-	/** Render priority/order. */
+	/** 
+	* Render priority/order. The higher the number the later region will be applied. 
+	* A region with Priority 1 will be rendered before a region with Priority 10. 
+	* This property is hidden if priority is determined by distance from the camera (When Window CCR is being used). 
+	*/
 	UPROPERTY(BlueprintReadWrite, EditAnywhere, Category="Color Correction")
 	int32 Priority;
 
diff --git a/Engine/Plugins/Experimental/ColorCorrectRegions/Source/ColorCorrectRegions/Public/ColorCorrectRegionsPostProcessMaterial.h b/Engine/Plugins/Experimental/ColorCorrectRegions/Source/ColorCorrectRegions/Public/ColorCorrectRegionsPostProcessMaterial.h
--- a/Engine/Plugins/Experimental/ColorCorrectRegions/Source/ColorCorrectRegions/Public/ColorCorrectRegionsPostProcessMaterial.h
+++ b/Engine/Plugins/Experimental/ColorCorrectRegions/Source/ColorCorrectRegions/Public/ColorCorrectRegionsPostProcessMaterial.h
@@ -136,16 +136,8 @@ public:
 	class FShaderType : SHADER_PERMUTATION_ENUM_CLASS("SHAPE_TYPE", EColorCorrectRegionsType);
 	class FTemperatureType : SHADER_PERMUTATION_ENUM_CLASS("TEMPERATURE_TYPE", ETemperatureType);
 	class FAdvancedShader : SHADER_PERMUTATION_BOOL("ADVANCED_CC");
-	class FDisplayBoundingRect : SHADER_PERMUTATION_BOOL("CCR_SHADER_DISPLAY_BOUNDING_RECT");
-	class FClipPixelsOutsideAABB : SHADER_PERMUTATION_BOOL("CLIP_PIXELS_OUTSIDE_AABB");
 
-	/** 
-	* On lower scalability settings Scene texture has only 3 channels.
-	* Which means we cannot sample it for opacity and need to get it from a different source.
-	*/
-	class FSampleOpacityFromGbuffer : SHADER_PERMUTATION_BOOL("SAMPLE_OPACITY_FROM_GBUFFER");
-
-	using FPermutationDomain = TShaderPermutationDomain<FShaderType, FTemperatureType, FAdvancedShader, FDisplayBoundingRect, FClipPixelsOutsideAABB, FSampleOpacityFromGbuffer>;
+	using FPermutationDomain = TShaderPermutationDomain<FShaderType, FTemperatureType, FAdvancedShader>;
 
 	FColorCorrectRegionMaterialPS() = default;
 	FColorCorrectRegionMaterialPS(const ShaderMetaType::CompiledShaderInitializerType & Initializer)
diff --git a/Engine/Plugins/Experimental/ColorCorrectRegions/Source/ColorCorrectRegions/Public/ColorCorrectRegionsSubsystem.h b/Engine/Plugins/Experimental/ColorCorrectRegions/Source/ColorCorrectRegions/Public/ColorCorrectRegionsSubsystem.h
--- a/Engine/Plugins/Experimental/ColorCorrectRegions/Source/ColorCorrectRegions/Public/ColorCorrectRegionsSubsystem.h
+++ b/Engine/Plugins/Experimental/ColorCorrectRegions/Source/ColorCorrectRegions/Public/ColorCorrectRegionsSubsystem.h
@@ -80,6 +80,9 @@ public:
 	/** Sorts regions based on priority. */
 	void SortRegionsByPriority();
 
+	/** Sorts regions based on distance from the camera. */
+	void SortRegionsByDistance(const FVector& ViewLocation);
+
 private:
 
 	/** Repopulates array of region actors. */
@@ -87,8 +90,11 @@ private:
 
 public:
 
-	/** Stores pointers to all ColorCorrectRegion Actors. */
-	TArray<AColorCorrectRegion*> Regions;
+	/** Stores pointers to ColorCorrectRegion Actors that use priority for sorting. */
+	TArray<AColorCorrectRegion*> RegionsPriorityBased;
+
+	/** Stores pointers to ColorCorrectRegion Actors that are based on distance from camera. */
+	TArray<AColorCorrectRegion*> RegionsDistanceBased;
 
 private:
 
diff --git a/Engine/Plugins/Experimental/ColorCorrectRegions/Source/ColorCorrectRegionsEditor/Private/ColorCorrectRegionCustomization.cpp b/Engine/Plugins/Experimental/ColorCorrectRegions/Source/ColorCorrectRegionsEditor/Private/ColorCorrectRegionCustomization.cpp
new file mode 100644
--- /dev/null
+++ b/Engine/Plugins/Experimental/ColorCorrectRegions/Source/ColorCorrectRegionsEditor/Private/ColorCorrectRegionCustomization.cpp
@@ -0,0 +1,31 @@
+// Copyright Epic Games, Inc. All Rights Reserved.
+
+#include "ColorCorrectRegionCustomization.h"
+#include "ColorCorrectRegion.h"
+#include "DetailLayoutBuilder.h"
+
+TSharedRef<IDetailCustomization> FColorCorrectRegionDetails::MakeInstance()
+{
+	return MakeShareable(new FColorCorrectRegionDetails);
+}
+
+void FColorCorrectRegionDetails::CustomizeDetails(IDetailLayoutBuilder& DetailBuilder)
+{
+	bool bPriorityHidden = false;
+	for (const TWeakObjectPtr<UObject>& SelectedObject : DetailBuilder.GetSelectedObjects())
+	{
+		if (AColorCorrectRegion* CCR = Cast<AColorCorrectRegion>(SelectedObject.Get()))
+		{
+			if (CCR->Type == EColorCorrectRegionsType::Plane)
+			{
+				bPriorityHidden = true;
+				break;
+			}
+		}
+	}
+	if (bPriorityHidden)
+	{
+		TSharedRef<IPropertyHandle> TileLayersProperty = DetailBuilder.GetProperty(GET_MEMBER_NAME_CHECKED(AColorCorrectRegion, Priority));
+		DetailBuilder.HideProperty(TileLayersProperty);
+	}
+}
diff --git a/Engine/Plugins/Experimental/ColorCorrectRegions/Source/ColorCorrectRegionsEditor/Private/ColorCorrectRegionCustomization.h b/Engine/Plugins/Experimental/ColorCorrectRegions/Source/ColorCorrectRegionsEditor/Private/ColorCorrectRegionCustomization.h
new file mode 100644
--- /dev/null
+++ b/Engine/Plugins/Experimental/ColorCorrectRegions/Source/ColorCorrectRegionsEditor/Private/ColorCorrectRegionCustomization.h
@@ -0,0 +1,19 @@
+// Copyright Epic Games, Inc. All Rights Reserved.
+
+#pragma once
+
+#include "CoreMinimal.h"
+#include "Layout/Visibility.h"
+#include "IDetailCustomization.h"
+#include "PropertyHandle.h"
+
+
+class FColorCorrectRegionDetails : public IDetailCustomization
+{
+public:
+	/** Makes a new instance of this detail layout class for a specific detail view requesting it */
+	static TSharedRef<IDetailCustomization> MakeInstance();
+
+	/** IDetailCustomization interface */
+	virtual void CustomizeDetails(IDetailLayoutBuilder& DetailBuilder) override;
+};
diff --git a/Engine/Plugins/Experimental/ColorCorrectRegions/Source/ColorCorrectRegionsEditor/Private/ColorCorrectRegionsEditorModule.cpp b/Engine/Plugins/Experimental/ColorCorrectRegions/Source/ColorCorrectRegionsEditor/Private/ColorCorrectRegionsEditorModule.cpp
--- a/Engine/Plugins/Experimental/ColorCorrectRegions/Source/ColorCorrectRegionsEditor/Private/ColorCorrectRegionsEditorModule.cpp
+++ b/Engine/Plugins/Experimental/ColorCorrectRegions/Source/ColorCorrectRegionsEditor/Private/ColorCorrectRegionsEditorModule.cpp
@@ -4,6 +4,7 @@
 #include "Editor/PlacementMode/Public/IPlacementModeModule.h"
 #include "ColorCorrectRegion.h"
 #include "ColorCorrectRegionsStyle.h"
+#include "ColorCorrectRegionCustomization.h"
 #include "ActorFactories/ActorFactoryBlueprint.h"
 
 #define LOCTEXT_NAMESPACE "FColorCorrectRegionsModule"
@@ -12,6 +13,9 @@ void FColorCorrectRegionsEditorModule::StartupModule()
 {
 	IPlacementModeModule::Get().OnPlacementModeCategoryRefreshed().AddRaw(this, &FColorCorrectRegionsEditorModule::OnPlacementModeRefresh);
 	FColorCorrectRegionsStyle::Initialize();
+
+	FPropertyEditorModule& PropertyModule = FModuleManager::LoadModuleChecked<FPropertyEditorModule>("PropertyEditor");
+	PropertyModule.RegisterCustomClassLayout(AColorCorrectRegion::StaticClass()->GetFName(), FOnGetDetailCustomizationInstance::CreateStatic(&FColorCorrectRegionDetails::MakeInstance));
 }
 
 void FColorCorrectRegionsEditorModule::OnPlacementModeRefresh(FName CategoryName)
diff --git a/Engine/Plugins/Experimental/ColorCorrectRegions/Source/ColorCorrectRegionsEditor/Public/ColorCorrectRegionsEditorModule.h b/Engine/Plugins/Experimental/ColorCorrectRegions/Source/ColorCorrectRegionsEditor/Public/ColorCorrectRegionsEditorModule.h
--- a/Engine/Plugins/Experimental/ColorCorrectRegions/Source/ColorCorrectRegionsEditor/Public/ColorCorrectRegionsEditorModule.h
+++ b/Engine/Plugins/Experimental/ColorCorrectRegions/Source/ColorCorrectRegionsEditor/Public/ColorCorrectRegionsEditorModule.h
@@ -17,4 +17,5 @@ public:
 	/** IModuleInterface implementation end */
 
 	void OnPlacementModeRefresh(FName CategoryName);
+	void OnPostEngineInit();
 };
