CCR: Removing SCOPED_GPU_STAT in case it becomes a blocker issue (crash) due to parallelized task rendering execution causing failed assertions.

diff --git a/Engine/Plugins/Experimental/ColorCorrectRegions/Source/ColorCorrectRegions/Private/ColorCorrectRegionsSceneViewExtension.cpp b/Engine/Plugins/Experimental/ColorCorrectRegions/Source/ColorCorrectRegions/Private/ColorCorrectRegionsSceneViewExtension.cpp
--- a/Engine/Plugins/Experimental/ColorCorrectRegions/Source/ColorCorrectRegions/Private/ColorCorrectRegionsSceneViewExtension.cpp
+++ b/Engine/Plugins/Experimental/ColorCorrectRegions/Source/ColorCorrectRegions/Private/ColorCorrectRegionsSceneViewExtension.cpp
@@ -21,9 +21,7 @@
 //Set this to 1 to see the clipping region.
 #define ColorCorrectRegions_SHADER_DISPLAY_BOUNDING_RECT 0
 
-DECLARE_GPU_STAT_NAMED(ColorCorrectRegion_Render, TEXT("ColorCorrectRegion_Render"));
-DECLARE_GPU_STAT_NAMED(ColorCorrectRegion_Copy, TEXT("ColorCorrectRegion_Copy"));
-DECLARE_GPU_STAT_NAMED(ColorCorrectRegion_Clear, TEXT("ColorCorrectRegion_Clear"));
+DECLARE_GPU_STAT_NAMED(ColorCorrectRegion, TEXT("ColorCorrectRegion"));
 
 namespace
 {
@@ -443,7 +441,6 @@ namespace
 				ERDGPassFlags::Raster,
 				[&View, ScreenPassVS, CopyPixelShader, RegionViewport, Parameters, DefaultBlendState](FRHICommandList& RHICmdList)
 				{
-					SCOPED_GPU_STAT(RHICmdList, ColorCorrectRegion_Clear);
 					DrawScreenPass(
 						RHICmdList,
 						static_cast<const FViewInfo&>(View),
@@ -477,8 +474,6 @@ namespace
 			CCHighlights,
 			bIsAdvanced](FRHICommandList& RHICmdList)
 			{
-				SCOPED_GPU_STAT(RHICmdList, ColorCorrectRegion_Render);
-
 				DrawScreenPass(
 					RHICmdList,
 					static_cast<const FViewInfo&>(View),
@@ -529,7 +524,6 @@ namespace
 			ERDGPassFlags::Raster,
 			[&View, ScreenPassVS, CopyPixelShader, RegionViewport, Parameters, CopyBlendState](FRHICommandList& RHICmdList)
 			{
-				SCOPED_GPU_STAT(RHICmdList, ColorCorrectRegion_Copy);
 				DrawScreenPass(
 					RHICmdList,
 					static_cast<const FViewInfo&>(View),
