Color Correct Regions: Added tint control. Contingent on https://p4-swarm.epicgames.net/reviews/20907136

diff --git a/Engine/Plugins/Experimental/ColorCorrectRegions/Shaders/Private/ColorCorrectRegionsColorTemperature.ush b/Engine/Plugins/Experimental/ColorCorrectRegions/Shaders/Private/ColorCorrectRegionsColorTemperature.ush
deleted file mode 100644
diff --git a/Engine/Plugins/Experimental/ColorCorrectRegions/Shaders/Private/ColorCorrectRegionsShader.usf b/Engine/Plugins/Experimental/ColorCorrectRegions/Shaders/Private/ColorCorrectRegionsShader.usf
--- a/Engine/Plugins/Experimental/ColorCorrectRegions/Shaders/Private/ColorCorrectRegionsShader.usf
+++ b/Engine/Plugins/Experimental/ColorCorrectRegions/Shaders/Private/ColorCorrectRegionsShader.usf
@@ -9,9 +9,8 @@
 #include "/Engine/Private/SceneData.ush"
 #include "/Engine/Private/ScreenPass.ush"
 #include "/Engine/Private/SceneTextureParameters.ush"
+#include "/Engine/Private/TonemapCommon.ush"
 
-
-#include "/ColorCorrectRegionsShaders/Private/ColorCorrectRegionsColorTemperature.ush"
 #include "/ColorCorrectRegionsShaders/Private/ColorCorrectRegionsCommon.usf"
 
 //////////////////////////////////////////////////////////////////////////
@@ -66,6 +65,7 @@ struct FExternalExpressions
 #endif
 
 	float WhiteTemp;
+	float Tint;
 	float Inner;
 	float Outer;
 
@@ -79,7 +79,6 @@ struct FExternalExpressions
 	float4 SceneTexture;
 };
 
-
 MaterialFloat3 ProcessColor(in FExternalExpressions ExternalExpressions)
 {
 	MaterialFloat3 WorkingColor = sRGBtoAP1(ExternalExpressions.SceneTexture);
@@ -92,14 +91,14 @@ MaterialFloat3 ProcessColor(in FExternalExpressions ExternalExpressions)
 		MaterialFloat BasicColorGrade_Mul = dot(BasicColorGrade, MaterialFloat3(0.30000001, 0.58999997, 0.11000000));
 		// In the future we might want to add a threshold width to increase or reduce interpolation area between shadows and midtones.
 		MaterialFloat Local5 = smoothstep(0.0 /*ExternalExpressions.ShadowMax - SHADOW_INTERPOLATION_THRESHOLD_WIDTH*/, ExternalExpressions.ShadowMax, BasicColorGrade_Mul);
-		MaterialFloat Local6 = (1.00000000 - Local5);
+		MaterialFloat Local6 = (1.0 - Local5);
 		MaterialFloat3 Local7 = (Local3 * Local6);
 		MaterialFloat3 Local8 = ComposureColorGradePerRange(BasicColorGrade, ExternalExpressions.ColorSaturationMidtones, ExternalExpressions.ColorContrastMidtones, ExternalExpressions.ColorGammaMidtones, ExternalExpressions.ColorGainMidtones, ExternalExpressions.ColorOffsetMidtones);
 		// Interpolation threshold width specifies the upper bound for the soft interpolated area between midtones and highlights. 
 		// In the future we want to add control that allows user to control this width.
 		MaterialFloat Local9 = smoothstep(ExternalExpressions.HighlightsMin, ExternalExpressions.HighlightsMin + HIGHLIGHT_INTERPOLATION_THRESHOLD_WIDTH, BasicColorGrade_Mul);
 		MaterialFloat Local10 = (Local6 + Local9);
-		MaterialFloat Local11 = (1.00000000 - Local10);
+		MaterialFloat Local11 = (1.0 - Local10);
 		MaterialFloat3 Local12 = (Local8 * Local11);
 		MaterialFloat3 Local13 = (Local7 + Local12);
 		MaterialFloat3 Local14 = ComposureColorGradePerRange(BasicColorGrade, ExternalExpressions.ColorSaturationHighlights, ExternalExpressions.ColorContrastHighlights, ExternalExpressions.ColorGammaHighlights, ExternalExpressions.ColorGainHighlights, ExternalExpressions.ColorOffsetHighlights);
@@ -113,32 +112,34 @@ MaterialFloat3 ProcessColor(in FExternalExpressions ExternalExpressions)
 	MaterialFloat3 AfterTemperature = FinalColor;
 
 #if TEMPERATURE_TYPE == TEMPERATURE_TYPE_LEGACY
-	MaterialFloat Local4 = smoothstep(0.00000000, 13000.00000000, ExternalExpressions.WhiteTemp);
-	MaterialFloat Local5 = (Local4 * 2.00000000 - 1.00000000);
+	MaterialFloat Local4 = smoothstep(0.0, 13000.0, ExternalExpressions.WhiteTemp);
+	MaterialFloat Local5 = (Local4 * 2.0 - 1.0);
 
 	// MF Artistic temperature 
 	{
-		MaterialFloat Local7 = (Local5 * 3.00000000);
-		MaterialFloat Local8 = clamp(Local7, 0.00000000, 5.00000000);
+		MaterialFloat Local7 = (Local5 * 3.0);
+		MaterialFloat Local8 = clamp(Local7, 0.0, 5.0);
 		MaterialFloat Local9 = (FinalColor.b * Local8);
 		MaterialFloat Local10 = (Local9 * 0.36666000);
 		MaterialFloat Local11 = (FinalColor.r + Local10);
-		MaterialFloat Local12 = (1.00000000 - Local8);
+		MaterialFloat Local12 = (1.0 - Local8);
 		MaterialFloat Local13 = (Local12 * FinalColor.b);
-		MaterialFloat Local14 = (Local5 * -3.00000000);
-		MaterialFloat Local15 = clamp(Local14, 0.00000000, 5.00000000);
-		MaterialFloat Local16 = (1.00000000 - Local15);
+		MaterialFloat Local14 = (Local5 * -3.0);
+		MaterialFloat Local15 = clamp(Local14, 0.0, 5.0);
+		MaterialFloat Local16 = (1.0 - Local15);
 		MaterialFloat Local17 = (Local16 * FinalColor.r);
 		MaterialFloat Local18 = (FinalColor.r * Local15);
 		MaterialFloat Local19 = (Local18 * 2.72000003);
 		MaterialFloat Local20 = (FinalColor.b + Local19);
-		MaterialFloat Local21 = smoothstep(0.00010000, 0.00000000, Local5);
-		AfterTemperature = lerp(MaterialFloat3(Local11, FinalColor.g, Local13), MaterialFloat3(Local17, FinalColor.g, Local20), Local21);
+		MaterialFloat Local21 = smoothstep(0.0001, 0.0, Local5);
+		
+		AfterTemperature = lerp(MaterialFloat3(Local11, FinalColor.g, Local13), MaterialFloat3(Local17, FinalColor.g, Local20), Local21); 
+		AfterTemperature.g = AfterTemperature.g - ExternalExpressions.Tint; 
 	}
 #elif TEMPERATURE_TYPE == TEMPERATURE_TYPE_WHITE_BALANCE
-	AfterTemperature = WhiteBalance(ExternalExpressions.WhiteTemp, FinalColor, false /*SwapSourceAndD65White*/);
+	AfterTemperature = WhiteBalance(FinalColor, ExternalExpressions.WhiteTemp, ExternalExpressions.Tint, true /*bIsTemperatureWhiteBalance*/);
 #elif TEMPERATURE_TYPE == TEMPERATURE_TYPE_COLOR_TEMP
-	AfterTemperature = WhiteBalance(ExternalExpressions.WhiteTemp, FinalColor, true /*SwapSourceAndD65White*/);
+	AfterTemperature = WhiteBalance(FinalColor, ExternalExpressions.WhiteTemp, ExternalExpressions.Tint, false /*bIsTemperatureWhiteBalance*/);
 #endif
 	return AfterTemperature;
 }
@@ -231,7 +232,7 @@ float GetShapeAlpha(in out FMaterialPixelParameters Parameters, in FExternalExpr
 		Shape = smoothstep(ExternalExpressions.Inner, ExternalExpressions.Outer, Distance);
 	}
 #endif
-	float ShapeAlpha = (1.00000000 - Shape);
+	float ShapeAlpha = (1.0 - Shape);
 	return ShapeAlpha;
 }
 #endif
@@ -254,11 +255,11 @@ MaterialFloat4 ProcessRegion(in out FMaterialPixelParameters Parameters, in FExt
 
 	/* Falloff */
 	float ShapeWithFalloff = PositiveClampedPow(ShapeAlpha, ExternalExpressions.Falloff);
-	float ShapeWithFalloffInverse = (1.00000000 - ShapeWithFalloff);
+	float ShapeWithFalloffInverse = (1.0 - ShapeWithFalloff);
 
 	float FinalShape = lerp(ShapeWithFalloff, ShapeWithFalloffInverse, ExternalExpressions.Invert);
 
-	float FinalAlpha = lerp(0.00000000, ExternalExpressions.Intensity, FinalShape);
+	float FinalAlpha = lerp(0.0, ExternalExpressions.Intensity, FinalShape);
 
 	MaterialFloat3 FinalColor = lerp(ExternalExpressions.SceneTexture.rgb, AfterTemperature, FinalAlpha);
 
@@ -274,7 +275,7 @@ MaterialFloat4 ProcessRegion(in out FMaterialPixelParameters Parameters, in FExt
 void CalcPixelMaterialInputs(in out FMaterialPixelParameters Parameters, in out FPixelMaterialInputs PixelMaterialInputs)
 {
     // The Normal is a special case as it might have its own expressions and also be used to calculate other inputs, so perform the assignment here
-	PixelMaterialInputs.Normal = MaterialFloat3(0.00000000, 0.00000000, 1.00000000);
+	PixelMaterialInputs.Normal = MaterialFloat3(0.0, 0.0, 1.0);
 
     // Note that here MaterialNormal can be in world space or tangent space
 	half3 MaterialNormal = GetMaterialNormal(Parameters, PixelMaterialInputs);
@@ -299,6 +300,7 @@ void CalcPixelMaterialInputs(in out FMaterialPixelParameters Parameters, in out
 		ExtExpressions.ColorGain = ColorCorrectBase.ColorGain;
 		ExtExpressions.ColorOffset = ColorCorrectBase.ColorOffset;
 		ExtExpressions.WhiteTemp = RegionData.WhiteTemp;
+		ExtExpressions.Tint = RegionData.Tint;
 		ExtExpressions.Inner = RegionData.Inner;
 		ExtExpressions.Outer = RegionData.Outer;
 		ExtExpressions.Falloff = RegionData.Falloff;
@@ -342,17 +344,17 @@ void CalcPixelMaterialInputs(in out FMaterialPixelParameters Parameters, in out
 	PixelMaterialInputs.Opacity = 1.0;
 #endif
 
-	PixelMaterialInputs.OpacityMask = 1.00000000;
-	PixelMaterialInputs.BaseColor = MaterialFloat3(0.00000000, 0.00000000, 0.00000000);
-	PixelMaterialInputs.Metallic = 0.00000000;
+	PixelMaterialInputs.OpacityMask = 1.0;
+	PixelMaterialInputs.BaseColor = MaterialFloat3(0.0, 0.0, 0.0);
+	PixelMaterialInputs.Metallic = 0.0;
 	PixelMaterialInputs.Specular = 0.50000000;
 	PixelMaterialInputs.Roughness = 0.50000000;
-	PixelMaterialInputs.Anisotropy = 0.00000000;
-	PixelMaterialInputs.Tangent = MaterialFloat3(1.00000000, 0.00000000, 0.00000000);
+	PixelMaterialInputs.Anisotropy = 0.0;
+	PixelMaterialInputs.Tangent = MaterialFloat3(1.0, 0.0, 0.0);
 	PixelMaterialInputs.Subsurface = 0;
-	PixelMaterialInputs.AmbientOcclusion = 1.00000000;
+	PixelMaterialInputs.AmbientOcclusion = 1.0;
 	PixelMaterialInputs.Refraction = 0;
-	PixelMaterialInputs.PixelDepthOffset = 0.00000000;
+	PixelMaterialInputs.PixelDepthOffset = 0.0;
 	PixelMaterialInputs.ShadingModel = 0;
 	Parameters.WorldTangent = 0;
 }
diff --git a/Engine/Plugins/Experimental/ColorCorrectRegions/Source/ColorCorrectRegions/Private/ColorCorrectRegion.cpp b/Engine/Plugins/Experimental/ColorCorrectRegions/Source/ColorCorrectRegions/Private/ColorCorrectRegion.cpp
--- a/Engine/Plugins/Experimental/ColorCorrectRegions/Source/ColorCorrectRegions/Private/ColorCorrectRegion.cpp
+++ b/Engine/Plugins/Experimental/ColorCorrectRegions/Source/ColorCorrectRegions/Private/ColorCorrectRegion.cpp
@@ -18,7 +18,9 @@ AColorCorrectRegion::AColorCorrectRegion(const FObjectInitializer& ObjectInitial
 	, Outer(1.0)
 	, Falloff(1.0)
 	, Invert(false)
+	, TemperatureType(EColorCorrectRegionTemperatureType::WhiteBalance)
 	, Temperature(6500)
+	, Tint(0)
 	, Enabled(true)
 	, ExcludeStencil(false)
 	, ColorCorrectRegionsSubsystem(nullptr)
@@ -64,7 +66,8 @@ AColorCorrectRegion::AColorCorrectRegion(const FObjectInitializer& ObjectInitial
 			SpriteComponent->SetRelativeLocationAndRotation(FVector::ZeroVector, FRotator::ZeroRotator);
 			SpriteComponent->SetMobility(EComponentMobility::Movable);
 			SpriteComponent->bHiddenInGame = true;
-			
+			SpriteComponent->bIsScreenSizeScaled = true;
+
 			SpriteComponent->AttachToComponent(RootComponent, FAttachmentTransformRules::KeepRelativeTransform);
 		}
 	}
diff --git a/Engine/Plugins/Experimental/ColorCorrectRegions/Source/ColorCorrectRegions/Private/ColorCorrectRegionsSceneViewExtension.cpp b/Engine/Plugins/Experimental/ColorCorrectRegions/Source/ColorCorrectRegions/Private/ColorCorrectRegionsSceneViewExtension.cpp
--- a/Engine/Plugins/Experimental/ColorCorrectRegions/Source/ColorCorrectRegions/Private/ColorCorrectRegionsSceneViewExtension.cpp
+++ b/Engine/Plugins/Experimental/ColorCorrectRegions/Source/ColorCorrectRegions/Private/ColorCorrectRegionsSceneViewExtension.cpp
@@ -224,7 +224,7 @@ namespace
 
 
 	bool RenderRegion
-	(FRDGBuilder& GraphBuilder
+		( FRDGBuilder& GraphBuilder
 		, const FSceneView& View
 		, const FPostProcessingInputs& Inputs
 		, const FSceneViewFamily& ViewFamily
@@ -239,6 +239,7 @@ namespace
 		, FGlobalShaderMap* GlobalShaderMap
 		, FRHIBlendState* DefaultBlendState)
 	{
+		SCOPED_GPU_STAT(GraphBuilder.RHICmdList, ColorCorrectRegion);
 		FRHIDepthStencilState* DepthStencilState = FScreenPassPipelineState::FDefaultDepthStencilState::GetRHI();
 
 		/* If Region is pending for kill, invisible or disabled we don't need to render it.
@@ -353,9 +354,10 @@ namespace
 
 		TShaderMapRef<FColorCorrectRegionMaterialVS> VertexShader(GlobalShaderMap);
 		const float DefaultTemperature = 6500;
+		const float DefaultTint = 0;
 
 		// If temperature is default we don't want to do the calculations.
-		FColorCorrectRegionMaterialPS::ETemperatureType TemperatureType = FMath::IsNearlyEqual(Region->Temperature, DefaultTemperature)
+		FColorCorrectRegionMaterialPS::ETemperatureType TemperatureType = FMath::IsNearlyEqual(Region->Temperature, DefaultTemperature) && FMath::IsNearlyEqual(Region->Tint, DefaultTint)
 			? FColorCorrectRegionMaterialPS::ETemperatureType::Disabled
 			: static_cast<FColorCorrectRegionMaterialPS::ETemperatureType>(Region->TemperatureType);
 		TShaderMapRef<FColorCorrectRegionMaterialPS> PixelShader = GetRegionShader(GlobalShaderMap, Region->Type, TemperatureType, bIsAdvanced);
@@ -379,6 +381,7 @@ namespace
 			RegionData.Scale = (FVector3f)Region->GetActorScale() * ScaleMultiplier;
 
 			RegionData.WhiteTemp = Region->Temperature;
+			RegionData.Tint = Region->Tint;
 			// Inner could be larger than outer, in which case we need to make sure these are swapped.
 			RegionData.Inner = FMath::Min<float>(Region->Outer, Region->Inner);
 			RegionData.Outer = FMath::Max<float>(Region->Outer, Region->Inner);
diff --git a/Engine/Plugins/Experimental/ColorCorrectRegions/Source/ColorCorrectRegions/Public/ColorCorrectRegion.h b/Engine/Plugins/Experimental/ColorCorrectRegions/Source/ColorCorrectRegions/Public/ColorCorrectRegion.h
--- a/Engine/Plugins/Experimental/ColorCorrectRegions/Source/ColorCorrectRegions/Public/ColorCorrectRegion.h
+++ b/Engine/Plugins/Experimental/ColorCorrectRegions/Source/ColorCorrectRegions/Public/ColorCorrectRegion.h
@@ -84,9 +84,13 @@ public:
 	EColorCorrectRegionTemperatureType TemperatureType;
 
 	/** Color correction temperature. */
-	UPROPERTY(BlueprintReadWrite, EditAnywhere, Category = "Color Correction")
+	UPROPERTY(BlueprintReadWrite, EditAnywhere, Category = "Color Correction", meta = (UIMin = "1500.0", UIMax = "15000.0"))
 	float Temperature;
 
+	/** Color temperature tint. */
+	UPROPERTY(BlueprintReadWrite, EditAnywhere, Category = "Color Correction", meta=(UIMin = "-1.0", UIMax = "1.0"))
+	float Tint;
+
 	/** Color correction settings. */
 	UPROPERTY(BlueprintReadWrite, EditAnywhere, Category = "Color Correction")
 	FColorGradingSettings ColorGradingSettings;
diff --git a/Engine/Plugins/Experimental/ColorCorrectRegions/Source/ColorCorrectRegions/Public/ColorCorrectRegionsPostProcessMaterial.h b/Engine/Plugins/Experimental/ColorCorrectRegions/Source/ColorCorrectRegions/Public/ColorCorrectRegionsPostProcessMaterial.h
--- a/Engine/Plugins/Experimental/ColorCorrectRegions/Source/ColorCorrectRegions/Public/ColorCorrectRegionsPostProcessMaterial.h
+++ b/Engine/Plugins/Experimental/ColorCorrectRegions/Source/ColorCorrectRegions/Public/ColorCorrectRegionsPostProcessMaterial.h
@@ -27,6 +27,7 @@ BEGIN_GLOBAL_SHADER_PARAMETER_STRUCT(FCCRRegionDataInputParameter, )
 	SHADER_PARAMETER(FVector3f, Scale)
 
 	SHADER_PARAMETER(float, WhiteTemp)
+	SHADER_PARAMETER(float, Tint)
 	SHADER_PARAMETER(float, Inner)
 	SHADER_PARAMETER(float, Outer)
 	SHADER_PARAMETER(float, Falloff)
