ColorCorrectRegions: Working colorspace is now fed to CCR and used in calculating white balance.

diff --git a/Engine/Plugins/Experimental/ColorCorrectRegions/Shaders/Private/ColorCorrectRegionsShader.usf b/Engine/Plugins/Experimental/ColorCorrectRegions/Shaders/Private/ColorCorrectRegionsShader.usf
--- a/Engine/Plugins/Experimental/ColorCorrectRegions/Shaders/Private/ColorCorrectRegionsShader.usf
+++ b/Engine/Plugins/Experimental/ColorCorrectRegions/Shaders/Private/ColorCorrectRegionsShader.usf
@@ -137,9 +137,9 @@ MaterialFloat3 ProcessColor(in FExternalExpressions ExternalExpressions)
 		AfterTemperature.g = AfterTemperature.g - ExternalExpressions.Tint; 
 	}
 #elif TEMPERATURE_TYPE == TEMPERATURE_TYPE_WHITE_BALANCE
-	AfterTemperature = WhiteBalance(FinalColor, ExternalExpressions.WhiteTemp, ExternalExpressions.Tint, true /*bIsTemperatureWhiteBalance*/);
+	AfterTemperature = WhiteBalance(FinalColor, ExternalExpressions.WhiteTemp, ExternalExpressions.Tint, true /*bIsTemperatureWhiteBalance*/, (float3x3)WorkingColorSpace.ToXYZ, (float3x3)WorkingColorSpace.FromXYZ);
 #elif TEMPERATURE_TYPE == TEMPERATURE_TYPE_COLOR_TEMP
-	AfterTemperature = WhiteBalance(FinalColor, ExternalExpressions.WhiteTemp, ExternalExpressions.Tint, false /*bIsTemperatureWhiteBalance*/);
+	AfterTemperature = WhiteBalance(FinalColor, ExternalExpressions.WhiteTemp, ExternalExpressions.Tint, false /*bIsTemperatureWhiteBalance*/, (float3x3)WorkingColorSpace.ToXYZ, (float3x3)WorkingColorSpace.FromXYZ);
 #endif
 	return AfterTemperature;
 }
diff --git a/Engine/Plugins/Experimental/ColorCorrectRegions/Source/ColorCorrectRegions/Private/ColorCorrectRegionsSceneViewExtension.cpp b/Engine/Plugins/Experimental/ColorCorrectRegions/Source/ColorCorrectRegions/Private/ColorCorrectRegionsSceneViewExtension.cpp
--- a/Engine/Plugins/Experimental/ColorCorrectRegions/Source/ColorCorrectRegions/Private/ColorCorrectRegionsSceneViewExtension.cpp
+++ b/Engine/Plugins/Experimental/ColorCorrectRegions/Source/ColorCorrectRegions/Private/ColorCorrectRegionsSceneViewExtension.cpp
@@ -347,6 +347,7 @@ namespace
 		FCCRShaderInputParameters* PostProcessMaterialParameters = GraphBuilder.AllocParameters<FCCRShaderInputParameters>();
 		PostProcessMaterialParameters->RenderTargets[0] = BackBufferRenderTarget.GetRenderTargetBinding();
 
+		PostProcessMaterialParameters->WorkingColorSpace = GDefaultWorkingColorSpaceUniformBuffer.GetUniformBufferRef();
 		PostProcessMaterialParameters->PostProcessOutput = SceneTextureViewportParams;
 		PostProcessMaterialParameters->PostProcessInput[0] = SceneTextureInput;
 		PostProcessMaterialParameters->SceneTextures = SceneTextures;
diff --git a/Engine/Plugins/Experimental/ColorCorrectRegions/Source/ColorCorrectRegions/Public/ColorCorrectRegionsPostProcessMaterial.h b/Engine/Plugins/Experimental/ColorCorrectRegions/Source/ColorCorrectRegions/Public/ColorCorrectRegionsPostProcessMaterial.h
--- a/Engine/Plugins/Experimental/ColorCorrectRegions/Source/ColorCorrectRegions/Public/ColorCorrectRegionsPostProcessMaterial.h
+++ b/Engine/Plugins/Experimental/ColorCorrectRegions/Source/ColorCorrectRegions/Public/ColorCorrectRegionsPostProcessMaterial.h
@@ -79,6 +79,7 @@ END_GLOBAL_SHADER_PARAMETER_STRUCT()
 BEGIN_SHADER_PARAMETER_STRUCT(FCCRShaderInputParameters, )
 	SHADER_PARAMETER_STRUCT_REF(FViewUniformShaderParameters, View)
 	SHADER_PARAMETER_STRUCT_INCLUDE(FSceneTextureShaderParameters, SceneTextures)
+	SHADER_PARAMETER_STRUCT_REF(FWorkingColorSpaceShaderParameters, WorkingColorSpace)
 	SHADER_PARAMETER_STRUCT(FScreenPassTextureViewportParameters, PostProcessOutput)
 	SHADER_PARAMETER_STRUCT_ARRAY(FScreenPassTextureInput, PostProcessInput, [1])
 	RENDER_TARGET_BINDING_SLOTS()
