ColorCorrectRegions: Fixed a couple of few issues and minor cleanup.
- Stencil is now doing depth comparison against scene depth to avoid stencilling areas of occluding objects.
- Fixed an engine shutdown crash, when user deletes CCRs and shuts down the engine.
- Fixed an issue with Details panel not refreshing and not hiding priority property for Plane types.

diff --git a/Engine/Plugins/Experimental/ColorCorrectRegions/Shaders/Private/ColorCorrectRegionsCommon.usf b/Engine/Plugins/Experimental/ColorCorrectRegions/Shaders/Private/ColorCorrectRegionsCommon.usf
--- a/Engine/Plugins/Experimental/ColorCorrectRegions/Shaders/Private/ColorCorrectRegionsCommon.usf
+++ b/Engine/Plugins/Experimental/ColorCorrectRegions/Shaders/Private/ColorCorrectRegionsCommon.usf
@@ -82,10 +82,6 @@ SamplerState PostProcessInput_BilinearSampler;
  */
 struct FMaterialPixelParameters
 {
-
-	/** Interpolated vertex color, in linear color space. */
-	half4 VertexColor;
-
 	/** Normalized world space normal. */
 	half3 WorldNormal;
 
@@ -98,9 +94,6 @@ struct FMaterialPixelParameters
 	/** Normalized world space camera vector, which is the vector from the point being shaded to the camera position. */
 	half3 CameraVector;
 
-	/** World space light vector, only valid when rendering a light function. */
-	half3 LightVector;
-
 	/**
 	 * Like SV_Position (.xy is pixel position at pixel center, z:DeviceZ, .w:SceneDepth)
 	 * using shader generated value SV_POSITION
@@ -108,18 +101,14 @@ struct FMaterialPixelParameters
 	 */
 	MaterialFloat4 SvPosition;
 
-	/** Post projection position reconstructed from SvPosition, before the divide by W. left..top -1..1, bottom..top -1..1  within the viewport, W is the SceneDepth */
-	MaterialFloat4 ScreenPosition;
-
-	half UnMirrored;
+	/* Scene depth in world units. */
+	half SceneDepth;
 
-	half TwoSidedSign;
+	/* Custom depth in world units. */
+	half CustomDepth;
 
-	/**
-	 * Orthonormal rotation-only transform from tangent space to world space
-	 * The transpose(TangentToWorld) is WorldToTangent, and TangentToWorld[2] is WorldVertexNormal
-	 */
-	half3x3 TangentToWorld;
+	/** Post projection position reconstructed from SvPosition, before the divide by W. left..top -1..1, bottom..top -1..1  within the viewport, W is the SceneDepth */
+	MaterialFloat4 ScreenPosition;
 
 	/**
 	 * Interpolated worldspace position of this pixel
@@ -131,28 +120,6 @@ struct FMaterialPixelParameters
 	 * Interpolated worldspace position of this pixel, centered around the camera
 	 */
 	MaterialFloat3 WorldPosition_CamRelative;
-
-	/**
-	 * Interpolated worldspace position of this pixel, not including any world position offset or displacement.
-	 * Only valid if shader is compiled with NEEDS_WORLD_POSITION_EXCLUDING_SHADER_OFFSETS, otherwise just contains 0
-	 */
-	MaterialFloat3 WorldPosition_NoOffsets;
-
-	/**
-	 * Interpolated worldspace position of this pixel, not including any world position offset or displacement.
-	 * Only valid if shader is compiled with NEEDS_WORLD_POSITION_EXCLUDING_SHADER_OFFSETS, otherwise just contains 0
-	 */
-	MaterialFloat3 WorldPosition_NoOffsets_CamRelative;
-
-	/** Offset applied to the lighting position for translucency, used to break up aliasing artifacts. */
-	half3 LightingPositionOffset;
-
-	MaterialFloat AOMaterialMask;
-
-	// Index into View.PrimitiveSceneData
-	uint PrimitiveId;
-
-
 };
 
 
diff --git a/Engine/Plugins/Experimental/ColorCorrectRegions/Shaders/Private/ColorCorrectRegionsShader.usf b/Engine/Plugins/Experimental/ColorCorrectRegions/Shaders/Private/ColorCorrectRegionsShader.usf
--- a/Engine/Plugins/Experimental/ColorCorrectRegions/Shaders/Private/ColorCorrectRegionsShader.usf
+++ b/Engine/Plugins/Experimental/ColorCorrectRegions/Shaders/Private/ColorCorrectRegionsShader.usf
@@ -265,8 +265,11 @@ MaterialFloat4 ProcessRegion(in out FMaterialPixelParameters Parameters, in FExt
 
 	if (ExternalExpressions.ExcludeStencil > 0)
 	{
-		float4 CustomStencil = SceneTextureLookup(GetDefaultSceneTextureUV(Parameters, PPI_CustomStencil), PPI_CustomStencil, false);
-		FinalColor = lerp(FinalColor, ExternalExpressions.SceneTexture.rgb, clamp(CustomStencil.r, 0., 1.));
+		float CustomStencil = SceneTextureLookup(GetDefaultSceneTextureUV(Parameters, PPI_CustomStencil), PPI_CustomStencil, false).r;
+
+		const float Float_Eps = 0.00000001;
+		float UseStencil = lerp(0, clamp(CustomStencil, 0., 1.), step(abs(Parameters.SceneDepth - Parameters.CustomDepth), Float_Eps));
+		FinalColor = lerp(FinalColor, ExternalExpressions.SceneTexture.rgb, UseStencil);
 	}
 
 	return MaterialFloat4(FinalColor, FinalAlpha);
@@ -394,14 +397,11 @@ void BuildParameters(in FPostProcessMaterialVSToPS Input, inout FMaterialPixelPa
 	Parameters.AbsoluteWorldPosition = TranslatedWorldPosition.xyz - LWCHackToFloat(ResolvedView.PreViewTranslation).xyz;
 
 	Parameters.SvPosition = SvPosition;
+	Parameters.SceneDepth = CalcSceneDepth(ViewportUVToBufferUV(ViewportUV));
+	Parameters.CustomDepth = CalcSceneCustomDepth(ViewportUVToBufferUV(ViewportUV));
 	Parameters.ScreenPosition = ScreenPosition;
 
 	Parameters.CameraVector = normalize(-Parameters.WorldPosition_CamRelative.xyz);
-
-	/* Unused */
-	Parameters.LightVector = 0;
-	Parameters.VertexColor = 1;
-	Parameters.TwoSidedSign = 1.0f;
 }
 
 void MainPS(
diff --git a/Engine/Plugins/Experimental/ColorCorrectRegions/Source/ColorCorrectRegions/Private/ColorCorrectRegionSubsystem.cpp b/Engine/Plugins/Experimental/ColorCorrectRegions/Source/ColorCorrectRegions/Private/ColorCorrectRegionSubsystem.cpp
--- a/Engine/Plugins/Experimental/ColorCorrectRegions/Source/ColorCorrectRegions/Private/ColorCorrectRegionSubsystem.cpp
+++ b/Engine/Plugins/Experimental/ColorCorrectRegions/Source/ColorCorrectRegions/Private/ColorCorrectRegionSubsystem.cpp
@@ -116,6 +116,7 @@ void UColorCorrectRegionsSubsystem::OnActorDeleted(AActor* InActor)
 		)
 #endif
 	{
+		AsRegion->Cleanup();
 		FScopeLock RegionScopeLock(&RegionAccessCriticalSection);
 		RegionsPriorityBased.Remove(AsRegion);
 		RegionsDistanceBased.Remove(AsRegion);
diff --git a/Engine/Plugins/Experimental/ColorCorrectRegions/Source/ColorCorrectRegionsEditor/Private/ColorCorrectRegionCustomization.cpp b/Engine/Plugins/Experimental/ColorCorrectRegions/Source/ColorCorrectRegionsEditor/Private/ColorCorrectRegionCustomization.cpp
--- a/Engine/Plugins/Experimental/ColorCorrectRegions/Source/ColorCorrectRegionsEditor/Private/ColorCorrectRegionCustomization.cpp
+++ b/Engine/Plugins/Experimental/ColorCorrectRegions/Source/ColorCorrectRegionsEditor/Private/ColorCorrectRegionCustomization.cpp
@@ -23,9 +23,20 @@ void FColorCorrectRegionDetails::CustomizeDetails(IDetailLayoutBuilder& DetailBu
 			}
 		}
 	}
+
 	if (bPriorityHidden)
 	{
-		TSharedRef<IPropertyHandle> TileLayersProperty = DetailBuilder.GetProperty(GET_MEMBER_NAME_CHECKED(AColorCorrectRegion, Priority));
-		DetailBuilder.HideProperty(TileLayersProperty);
+		TSharedRef<IPropertyHandle> PriorityProperty = DetailBuilder.GetProperty(GET_MEMBER_NAME_CHECKED(AColorCorrectRegion, Priority));
+		DetailBuilder.HideProperty(PriorityProperty);
+	}
+
+	// Trick to force the viewport to update because depending on Type, Priority property could be hidden.
+	TSharedRef<IPropertyHandle> TypeProperty = DetailBuilder.GetProperty(GET_MEMBER_NAME_CHECKED(AColorCorrectRegion, Type));
+	if (TypeProperty->IsValidHandle())
+	{
+		TypeProperty->SetOnPropertyValueChanged(FSimpleDelegate::CreateLambda([&DetailBuilder]()
+			{
+				DetailBuilder.ForceRefreshDetails();
+			}));
 	}
 }
