ColorCorrectRegions: with the working color space now integrated, we also want to support dynamic to/from AP1 conversions.

diff --git a/Engine/Plugins/Experimental/ColorCorrectRegions/Shaders/Private/ColorCorrectRegionsCommon.usf b/Engine/Plugins/Experimental/ColorCorrectRegions/Shaders/Private/ColorCorrectRegionsCommon.usf
--- a/Engine/Plugins/Experimental/ColorCorrectRegions/Shaders/Private/ColorCorrectRegionsCommon.usf
+++ b/Engine/Plugins/Experimental/ColorCorrectRegions/Shaders/Private/ColorCorrectRegionsCommon.usf
@@ -174,31 +174,16 @@ struct FMaterialVertexParameters
 
 };
 
-// Uniform material expressions.
-MaterialFloat3 sRGBtoAP1(MaterialFloat4 InputColor)
+MaterialFloat3 ComposureColorGradePerRange(MaterialFloat3 ColorAP1, MaterialFloat4 ColorSaturation, MaterialFloat4 ColorContrast, MaterialFloat4 ColorGamma, MaterialFloat4 ColorGain, MaterialFloat4 ColorOffset)
 {
-	const MaterialFloat3x3 sRGB_2_AP1 = mul(XYZ_2_AP1_MAT, mul(D65_2_D60_CAT, sRGB_2_XYZ_MAT));
-	MaterialFloat3 ColorAP1 = mul(sRGB_2_AP1, InputColor.xyz);
+	MaterialFloat Luma = dot(ColorAP1, MaterialFloat3(0.2722287168, 0.6740817658, 0.0536895174));
+	ColorAP1 = max(0, lerp(Luma.xxx, ColorAP1, ColorSaturation.xyz * ColorSaturation.w));
+	ColorAP1 = pow(ColorAP1 * (5.55555555), ColorContrast.xyz * ColorContrast.w) * 0.18;
+	ColorAP1 = pow(ColorAP1, 1.0 / (ColorGamma.xyz * ColorGamma.w));
+	ColorAP1 = ColorAP1 * (ColorGain.xyz * ColorGain.w) + (ColorOffset.xyz + ColorOffset.w);
 	return ColorAP1;
 }
 
-MaterialFloat3 AP1toSRGB(MaterialFloat3 InputColor)
-{
-	const MaterialFloat3x3 AP1_2_sRGB = mul(XYZ_2_sRGB_MAT, mul(D60_2_D65_CAT, AP1_2_XYZ_MAT));
-	MaterialFloat3 ColorRGB = mul(AP1_2_sRGB, InputColor);
-	return ColorRGB;
-}
-
-MaterialFloat3 ComposureColorGradePerRange(MaterialFloat3 WorkingColor, MaterialFloat4 ColorSaturation, MaterialFloat4 ColorContrast, MaterialFloat4 ColorGamma, MaterialFloat4 ColorGain, MaterialFloat4 ColorOffset)
-{
-	MaterialFloat Luma = dot(WorkingColor, MaterialFloat3(0.2722287168, 0.6740817658, 0.0536895174));
-	WorkingColor = max(0, lerp(Luma.xxx, WorkingColor, ColorSaturation.xyz * ColorSaturation.w));
-	WorkingColor = pow(WorkingColor * (5.55555555), ColorContrast.xyz * ColorContrast.w) * 0.18;
-	WorkingColor = pow(WorkingColor, 1.0 / (ColorGamma.xyz * ColorGamma.w));
-	WorkingColor = WorkingColor * (ColorGain.xyz * ColorGain.w) + (ColorOffset.xyz + ColorOffset.w);
-	return WorkingColor;
-}
-
 
 /** Calculate a reflection vector about the specified world space normal. Optionally normalize this normal **/
 MaterialFloat3 ReflectionAboutCustomWorldNormal(FMaterialPixelParameters Parameters, MaterialFloat3 WorldNormal, bool bNormalizeInputNormal)
diff --git a/Engine/Plugins/Experimental/ColorCorrectRegions/Shaders/Private/ColorCorrectRegionsShader.usf b/Engine/Plugins/Experimental/ColorCorrectRegions/Shaders/Private/ColorCorrectRegionsShader.usf
--- a/Engine/Plugins/Experimental/ColorCorrectRegions/Shaders/Private/ColorCorrectRegionsShader.usf
+++ b/Engine/Plugins/Experimental/ColorCorrectRegions/Shaders/Private/ColorCorrectRegionsShader.usf
@@ -81,8 +81,8 @@ struct FExternalExpressions
 
 MaterialFloat3 ProcessColor(in FExternalExpressions ExternalExpressions)
 {
-	MaterialFloat3 WorkingColor = sRGBtoAP1(ExternalExpressions.SceneTexture);
-	MaterialFloat3 BasicColorGrade = ComposureColorGradePerRange(WorkingColor, ExternalExpressions.ColorSaturation, ExternalExpressions.ColorContrast, ExternalExpressions.ColorGamma, ExternalExpressions.ColorGain, ExternalExpressions.ColorOffset);
+	MaterialFloat3 ColorAP1 = mul((float3x3)WorkingColorSpace.ToAP1, ExternalExpressions.SceneTexture.xyz);
+	MaterialFloat3 BasicColorGrade = ComposureColorGradePerRange(ColorAP1, ExternalExpressions.ColorSaturation, ExternalExpressions.ColorContrast, ExternalExpressions.ColorGamma, ExternalExpressions.ColorGain, ExternalExpressions.ColorOffset);
 	MaterialFloat3 FinalColor = BasicColorGrade;
 #if ADVANCED_CC	
 	{
@@ -106,7 +106,7 @@ MaterialFloat3 ProcessColor(in FExternalExpressions ExternalExpressions)
 		FinalColor = (Local13 + Local15);
 	}
 #endif
-	FinalColor = AP1toSRGB(FinalColor);
+	FinalColor = mul((float3x3)WorkingColorSpace.FromAP1, FinalColor.xyz);
 
 	// White temperature calculation.
 	MaterialFloat3 AfterTemperature = FinalColor;
