Color Correct Regions:
- Added a new Color Correction Window Actor.
- Added a new CCW type - circle.
- Split Color correct regions into CCR and Color Correct Windows.
- changed default Temperature type to ColorTemperature instead of White balance.

diff --git a/Engine/Plugins/Experimental/ColorCorrectRegions/Shaders/Private/ColorCorrectRegionsShader.usf b/Engine/Plugins/Experimental/ColorCorrectRegions/Shaders/Private/ColorCorrectRegionsShader.usf
--- a/Engine/Plugins/Experimental/ColorCorrectRegions/Shaders/Private/ColorCorrectRegionsShader.usf
+++ b/Engine/Plugins/Experimental/ColorCorrectRegions/Shaders/Private/ColorCorrectRegionsShader.usf
@@ -20,7 +20,9 @@
 #define SHAPE_TYPE_BOX 1
 #define SHAPE_TYPE_CYLINDER 2
 #define SHAPE_TYPE_CONE 3
-#define SHAPE_TYPE_PLANE 4
+
+#define WINDOW_TYPE_PLANE 0
+#define WINDOW_TYPE_CIRCLE 1
 
 #define TEMPERATURE_TYPE_LEGACY 0
 #define TEMPERATURE_TYPE_WHITE_BALANCE 1
@@ -144,7 +146,7 @@ MaterialFloat3 ProcessColor(in FExternalExpressions ExternalExpressions)
 	return AfterTemperature;
 }
 
-#if SHAPE_TYPE == SHAPE_TYPE_PLANE
+#ifdef WINDOW_TYPE 
 float GetShapeAlpha(in out FMaterialPixelParameters Parameters, in FExternalExpressions ExternalExpressions)
 {
 	// Pixel position in plane space.
@@ -168,8 +170,14 @@ float GetShapeAlpha(in out FMaterialPixelParameters Parameters, in FExternalExpr
 		PositionOnPlane = CameraVectorPlaneSpace + IntersectionSide / dot(DirVector, PlaneNormal) * DirVector;
 	}
 
-	float2 Interpolator = smoothstep(ExternalExpressions.Outer, ExternalExpressions.Inner, abs(PositionOnPlane.rg)) * step(-RotatedTranslatedScaled.b * DirectionSign, 0.) * step(IntersectionSide, 0.);
-	float ShapeAlpha = Interpolator.r * step(0., Interpolator.g) * Interpolator.g * step(0., Interpolator.r);
+	// Checks if Z is in front or behind the plane.
+	float ClipZ = step(-RotatedTranslatedScaled.z * DirectionSign, 0.);
+#if WINDOW_TYPE == WINDOW_TYPE_PLANE
+	float2 Interpolator = smoothstep(ExternalExpressions.Outer, ExternalExpressions.Inner, abs(PositionOnPlane.xy)) * ClipZ * step(IntersectionSide, 0.);
+	float ShapeAlpha = Interpolator.x * step(0., Interpolator.y) * Interpolator.g * step(0., Interpolator.x);
+#elif WINDOW_TYPE == WINDOW_TYPE_CIRCLE
+	float ShapeAlpha = smoothstep(ExternalExpressions.Outer, ExternalExpressions.Inner, length(PositionOnPlane)) * ClipZ;
+#endif
 
 #if CLIP_PIXELS_OUTSIDE_AABB
 	// Clip pixel if it is outside the bounding box.
@@ -180,8 +188,7 @@ float GetShapeAlpha(in out FMaterialPixelParameters Parameters, in FExternalExpr
 
 	return ShapeAlpha;
 }
-
-#else
+#else //#ifdef WINDOW_TYPE 
 
 float GetShapeAlpha(in out FMaterialPixelParameters Parameters, in FExternalExpressions ExternalExpressions)
 {
@@ -235,7 +242,7 @@ float GetShapeAlpha(in out FMaterialPixelParameters Parameters, in FExternalExpr
 	float ShapeAlpha = (1.0 - Shape);
 	return ShapeAlpha;
 }
-#endif
+#endif //#ifdef WINDOW_TYPE 
 
 MaterialFloat4 ProcessRegion(in out FMaterialPixelParameters Parameters, in FExternalExpressions ExternalExpressions)
 {
diff --git a/Engine/Plugins/Experimental/ColorCorrectRegions/Source/ColorCorrectRegions/Private/ColorCorrectRegion.cpp b/Engine/Plugins/Experimental/ColorCorrectRegions/Source/ColorCorrectRegions/Private/ColorCorrectRegion.cpp
--- a/Engine/Plugins/Experimental/ColorCorrectRegions/Source/ColorCorrectRegions/Private/ColorCorrectRegion.cpp
+++ b/Engine/Plugins/Experimental/ColorCorrectRegions/Source/ColorCorrectRegions/Private/ColorCorrectRegion.cpp
@@ -18,7 +18,7 @@ AColorCorrectRegion::AColorCorrectRegion(const FObjectInitializer& ObjectInitial
 	, Outer(1.0)
 	, Falloff(1.0)
 	, Invert(false)
-	, TemperatureType(EColorCorrectRegionTemperatureType::WhiteBalance)
+	, TemperatureType(EColorCorrectRegionTemperatureType::ColorTemperature)
 	, Temperature(6500)
 	, Tint(0)
 	, Enabled(true)
diff --git a/Engine/Plugins/Experimental/ColorCorrectRegions/Source/ColorCorrectRegions/Private/ColorCorrectRegionSubsystem.cpp b/Engine/Plugins/Experimental/ColorCorrectRegions/Source/ColorCorrectRegions/Private/ColorCorrectRegionSubsystem.cpp
--- a/Engine/Plugins/Experimental/ColorCorrectRegions/Source/ColorCorrectRegions/Private/ColorCorrectRegionSubsystem.cpp
+++ b/Engine/Plugins/Experimental/ColorCorrectRegions/Source/ColorCorrectRegions/Private/ColorCorrectRegionSubsystem.cpp
@@ -1,5 +1,6 @@
 // Copyright Epic Games, Inc. All Rights Reserved.
 #include "ColorCorrectRegionsSubsystem.h"
+#include "ColorCorrectWindow.h"
 #include "Engine/World.h"
 #include "Kismet/GameplayStatics.h"
 #include "EngineUtils.h"
@@ -93,12 +94,13 @@ void UColorCorrectRegionsSubsystem::OnActorSpawned(AActor* InActor)
 		// We wouldn't have to do a check here except in case of nDisplay we need to populate this list during OnLevelsChanged 
 		// because nDisplay can release Actors while those are marked as BeginningPlay. Therefore we want to avoid 
 		// adding regions twice.
-		TArray<AColorCorrectRegion*>* RegionsToAddTo = CCRType == EColorCorrectRegionsType::Plane ? &RegionsDistanceBased : &RegionsPriorityBased;
+		bool bIsDistanceBased = Cast<AColorCorrectWindow>(InActor) != nullptr;
+		TArray<AColorCorrectRegion*>* RegionsToAddTo = bIsDistanceBased ? &RegionsDistanceBased : &RegionsPriorityBased;
 		if (!RegionsToAddTo->Contains(AsRegion))
 		{
 			RegionsToAddTo->Add(AsRegion);
 			// Distance based CCR can only be sorted on render, when View info is available.
-			if (CCRType != EColorCorrectRegionsType::Plane)
+			if (!bIsDistanceBased)
 			{
 				SortRegionsByPriority();
 			}
@@ -163,7 +165,7 @@ void UColorCorrectRegionsSubsystem::RefreshRegions()
 		AColorCorrectRegion* AsRegion = *It;
 		if (IsRegionValid(AsRegion, GetWorld()))
 		{
-			if (AsRegion->Type != EColorCorrectRegionsType::Plane)
+			if (!Cast<AColorCorrectWindow>(AsRegion))
 			{
 				RegionsPriorityBased.Add(AsRegion);
 			}
diff --git a/Engine/Plugins/Experimental/ColorCorrectRegions/Source/ColorCorrectRegions/Private/ColorCorrectRegionsPostProcessMaterial.cpp b/Engine/Plugins/Experimental/ColorCorrectRegions/Source/ColorCorrectRegions/Private/ColorCorrectRegionsPostProcessMaterial.cpp
--- a/Engine/Plugins/Experimental/ColorCorrectRegions/Source/ColorCorrectRegions/Private/ColorCorrectRegionsPostProcessMaterial.cpp
+++ b/Engine/Plugins/Experimental/ColorCorrectRegions/Source/ColorCorrectRegions/Private/ColorCorrectRegionsPostProcessMaterial.cpp
@@ -24,5 +24,6 @@ IMPLEMENT_GLOBAL_SHADER_PARAMETER_STRUCT(FCCRColorCorrectHighlightsParameter, "C
 
 IMPLEMENT_GLOBAL_SHADER(FColorCorrectRegionMaterialVS, "/ColorCorrectRegionsShaders/Private/ColorCorrectRegionsShader.usf", "MainVS", SF_Vertex);
 IMPLEMENT_GLOBAL_SHADER(FColorCorrectRegionMaterialPS, "/ColorCorrectRegionsShaders/Private/ColorCorrectRegionsShader.usf", "MainPS", SF_Pixel);
+IMPLEMENT_GLOBAL_SHADER(FColorCorrectWindowMaterialPS, "/ColorCorrectRegionsShaders/Private/ColorCorrectRegionsShader.usf", "MainPS", SF_Pixel);
 IMPLEMENT_GLOBAL_SHADER(FColorCorrectScreenPassVS, "/ColorCorrectRegionsShaders/Private/ColorCorrectRegionsScreenPass.usf", "MainVS", SF_Vertex);
 IMPLEMENT_GLOBAL_SHADER(FClearRectPS, "/ColorCorrectRegionsShaders/Private/ColorCorrectRegionsScreenPass.usf", "MainPS", SF_Pixel);
diff --git a/Engine/Plugins/Experimental/ColorCorrectRegions/Source/ColorCorrectRegions/Private/ColorCorrectRegionsSceneViewExtension.cpp b/Engine/Plugins/Experimental/ColorCorrectRegions/Source/ColorCorrectRegions/Private/ColorCorrectRegionsSceneViewExtension.cpp
--- a/Engine/Plugins/Experimental/ColorCorrectRegions/Source/ColorCorrectRegions/Private/ColorCorrectRegionsSceneViewExtension.cpp
+++ b/Engine/Plugins/Experimental/ColorCorrectRegions/Source/ColorCorrectRegions/Private/ColorCorrectRegionsSceneViewExtension.cpp
@@ -202,16 +202,26 @@ namespace
 				View.Family->EngineShowFlags.PostProcessMaterial;
 	}
 
-	// A helper function for getting the right shader.
-	TShaderMapRef<FColorCorrectRegionMaterialPS> GetRegionShader(const FGlobalShaderMap* GlobalShaderMap, EColorCorrectRegionsType RegionType, FColorCorrectRegionMaterialPS::ETemperatureType TemperatureType, bool bIsAdvanced)
+	// A helper function for getting the right shader for SDF based CCRs.
+	TShaderMapRef<FColorCorrectRegionMaterialPS> GetRegionShader(const FGlobalShaderMap* GlobalShaderMap, EColorCorrectRegionsType RegionType, FColorCorrectGenericPS::ETemperatureType TemperatureType, bool bIsAdvanced)
 	{
 		FColorCorrectRegionMaterialPS::FPermutationDomain PermutationVector;
-		PermutationVector.Set<FColorCorrectRegionMaterialPS::FAdvancedShader>(bIsAdvanced);
+		PermutationVector.Set<FColorCorrectGenericPS::FAdvancedShader>(bIsAdvanced);
+		PermutationVector.Set<FColorCorrectGenericPS::FTemperatureType>(TemperatureType);
 		PermutationVector.Set<FColorCorrectRegionMaterialPS::FShaderType>(static_cast<EColorCorrectRegionsType>(FMath::Min(static_cast<int32>(RegionType), static_cast<int32>(EColorCorrectRegionsType::MAX) - 1)));
-		PermutationVector.Set<FColorCorrectRegionMaterialPS::FTemperatureType>(TemperatureType);
 
 		return TShaderMapRef<FColorCorrectRegionMaterialPS>(GlobalShaderMap, PermutationVector);
-		;
+	}
+
+	// A helper function for getting the right shader for distance based CCRs.
+	TShaderMapRef<FColorCorrectWindowMaterialPS> GetWindowShader(const FGlobalShaderMap* GlobalShaderMap, EColorCorrectWindowType RegionType, FColorCorrectGenericPS::ETemperatureType TemperatureType, bool bIsAdvanced)
+	{
+		FColorCorrectWindowMaterialPS::FPermutationDomain PermutationVector;
+		PermutationVector.Set<FColorCorrectGenericPS::FAdvancedShader>(bIsAdvanced);
+		PermutationVector.Set<FColorCorrectGenericPS::FTemperatureType>(TemperatureType);
+		PermutationVector.Set<FColorCorrectWindowMaterialPS::FShaderType>(static_cast<EColorCorrectWindowType>(FMath::Min(static_cast<int32>(RegionType), static_cast<int32>(EColorCorrectWindowType::MAX) - 1)));
+
+		return TShaderMapRef<FColorCorrectWindowMaterialPS>(GlobalShaderMap, PermutationVector);
 	}
 
 	FVector4 Clamp(const FVector4 & VectorToClamp, float Min, float Max)
@@ -361,7 +371,16 @@ namespace
 		FColorCorrectRegionMaterialPS::ETemperatureType TemperatureType = FMath::IsNearlyEqual(Region->Temperature, DefaultTemperature) && FMath::IsNearlyEqual(Region->Tint, DefaultTint)
 			? FColorCorrectRegionMaterialPS::ETemperatureType::Disabled
 			: static_cast<FColorCorrectRegionMaterialPS::ETemperatureType>(Region->TemperatureType);
-		TShaderMapRef<FColorCorrectRegionMaterialPS> PixelShader = GetRegionShader(GlobalShaderMap, Region->Type, TemperatureType, bIsAdvanced);
+
+		TShaderRef<FColorCorrectGenericPS> PixelShader;
+		if (AColorCorrectWindow* CCWindow = Cast<AColorCorrectWindow>(Region))
+		{
+			PixelShader = GetWindowShader(GlobalShaderMap, CCWindow->WindowType, TemperatureType, bIsAdvanced);
+		}
+		else
+		{
+			PixelShader = GetRegionShader(GlobalShaderMap, Region->Type, TemperatureType, bIsAdvanced);
+		}
 
 		ClearUnusedGraphResources(VertexShader, PixelShader, PostProcessMaterialParameters);
 
diff --git a/Engine/Plugins/Experimental/ColorCorrectRegions/Source/ColorCorrectRegions/Private/ColorCorrectWindow.cpp b/Engine/Plugins/Experimental/ColorCorrectRegions/Source/ColorCorrectRegions/Private/ColorCorrectWindow.cpp
new file mode 100644
--- /dev/null
+++ b/Engine/Plugins/Experimental/ColorCorrectRegions/Source/ColorCorrectRegions/Private/ColorCorrectWindow.cpp
@@ -0,0 +1,61 @@
+// Copyright Epic Games, Inc. All Rights Reserved.
+
+#include "ColorCorrectWindow.h"
+#include "Components/StaticMeshComponent.h"
+#include "CoreMinimal.h"
+#include "Engine/Texture2D.h"
+#include "UObject/ConstructorHelpers.h"
+
+ENUM_RANGE_BY_COUNT(EColorCorrectWindowType, EColorCorrectWindowType::MAX)
+
+AColorCorrectWindow::AColorCorrectWindow(const FObjectInitializer& ObjectInitializer)
+	: Super(ObjectInitializer)
+	, WindowType(EColorCorrectWindowType::Plane)
+{
+	UMaterial* Material = LoadObject<UMaterial>(NULL, TEXT("/ColorCorrectRegions/Materials/M_ColorCorrectRegionTransparentPreview.M_ColorCorrectRegionTransparentPreview"), NULL, LOAD_None, NULL);
+	const TArray<UStaticMesh*> StaticMeshes =
+	{
+		Cast<UStaticMesh>(StaticLoadObject(UStaticMesh::StaticClass(), NULL, TEXT("/Engine/BasicShapes/Plane"))),
+		Cast<UStaticMesh>(StaticLoadObject(UStaticMesh::StaticClass(), NULL, TEXT("/Engine/BasicShapes/Cylinder")))
+	};
+	for (EColorCorrectWindowType CCWType : TEnumRange<EColorCorrectWindowType>())
+	{
+		UStaticMeshComponent* MeshComponent = CreateDefaultSubobject<UStaticMeshComponent>(*UEnum::GetValueAsString(CCWType));
+		MeshComponents.Add(MeshComponent);
+		MeshComponent->SetupAttachment(RootComponent);
+		MeshComponent->SetStaticMesh(StaticMeshes[static_cast<uint8>(CCWType)]);
+		MeshComponent->SetWorldScale3D(FVector(1., 1., 0.001));
+		MeshComponent->SetMaterial(0, Material);
+	}
+	SetMeshVisibilityForWindowType();
+}
+
+#if WITH_EDITOR
+void AColorCorrectWindow::PostEditChangeProperty(struct FPropertyChangedEvent& PropertyChangedEvent)
+{
+	Super::PostEditChangeProperty(PropertyChangedEvent);
+	const FName PropertyName = PropertyChangedEvent.GetPropertyName();
+
+	if (PropertyName == GET_MEMBER_NAME_CHECKED(AColorCorrectWindow, WindowType))
+	{
+		SetMeshVisibilityForWindowType();
+	}
+}
+#endif //WITH_EDITOR
+
+void AColorCorrectWindow::SetMeshVisibilityForWindowType()
+{
+	for (EColorCorrectWindowType CCWType : TEnumRange<EColorCorrectWindowType>())
+	{
+		uint8 TypeIndex = static_cast<uint8>(CCWType);
+
+		if (CCWType == WindowType)
+		{
+			MeshComponents[TypeIndex]->SetVisibility(true, true);
+		}
+		else
+		{
+			MeshComponents[TypeIndex]->SetVisibility(false, true);
+		}
+	}
+}
\ No newline at end of file
diff --git a/Engine/Plugins/Experimental/ColorCorrectRegions/Source/ColorCorrectRegions/Public/ColorCorrectRegion.h b/Engine/Plugins/Experimental/ColorCorrectRegions/Source/ColorCorrectRegions/Public/ColorCorrectRegion.h
--- a/Engine/Plugins/Experimental/ColorCorrectRegions/Source/ColorCorrectRegions/Public/ColorCorrectRegion.h
+++ b/Engine/Plugins/Experimental/ColorCorrectRegions/Source/ColorCorrectRegions/Public/ColorCorrectRegion.h
@@ -20,7 +20,6 @@ enum class EColorCorrectRegionsType : uint8
 	Box			UMETA(DisplayName = "Box"),
 	Cylinder	UMETA(DisplayName = "Cylinder"),
 	Cone		UMETA(DisplayName = "Cone"),
-	Plane		UMETA(DisplayName = "Plane (Window CCR)"),
 	MAX
 };
 
diff --git a/Engine/Plugins/Experimental/ColorCorrectRegions/Source/ColorCorrectRegions/Public/ColorCorrectRegionsPostProcessMaterial.h b/Engine/Plugins/Experimental/ColorCorrectRegions/Source/ColorCorrectRegions/Public/ColorCorrectRegionsPostProcessMaterial.h
--- a/Engine/Plugins/Experimental/ColorCorrectRegions/Source/ColorCorrectRegions/Public/ColorCorrectRegionsPostProcessMaterial.h
+++ b/Engine/Plugins/Experimental/ColorCorrectRegions/Source/ColorCorrectRegions/Public/ColorCorrectRegionsPostProcessMaterial.h
@@ -14,6 +14,7 @@
 #include "Runtime/RenderCore/Public/RenderGraphResources.h"
 #include "Runtime/RenderCore/Public/RenderGraphResources.h"
 #include "ColorCorrectRegion.h"
+#include "ColorCorrectWindow.h"
 
 // FScreenPassTextureViewportParameters and FScreenPassTextureInput
 #include "Runtime/Renderer/Private/ScreenPass.h"
@@ -121,7 +122,7 @@ public:
 	
 };
 
-class FColorCorrectRegionMaterialPS : public FColorCorrectRegionsPostProcessMaterialShader
+class FColorCorrectGenericPS : public FColorCorrectRegionsPostProcessMaterialShader
 {
 public:
 	UENUM(BlueprintType)
@@ -133,16 +134,12 @@ public:
 		Disabled,
 		MAX
 	};
-	DECLARE_GLOBAL_SHADER(FColorCorrectRegionMaterialPS);
 
-	class FShaderType : SHADER_PERMUTATION_ENUM_CLASS("SHAPE_TYPE", EColorCorrectRegionsType);
 	class FTemperatureType : SHADER_PERMUTATION_ENUM_CLASS("TEMPERATURE_TYPE", ETemperatureType);
 	class FAdvancedShader : SHADER_PERMUTATION_BOOL("ADVANCED_CC");
 
-	using FPermutationDomain = TShaderPermutationDomain<FShaderType, FTemperatureType, FAdvancedShader>;
-
-	FColorCorrectRegionMaterialPS() = default;
-	FColorCorrectRegionMaterialPS(const ShaderMetaType::CompiledShaderInitializerType & Initializer)
+	FColorCorrectGenericPS() = default;
+	FColorCorrectGenericPS(const ShaderMetaType::CompiledShaderInitializerType& Initializer)
 		: FColorCorrectRegionsPostProcessMaterialShader(Initializer)
 	{}
 
@@ -151,10 +148,44 @@ public:
 		FGlobalShader::SetParameters<FViewUniformShaderParameters>(RHICmdList, RHICmdList.GetBoundPixelShader(), View.ViewUniformBuffer);
 	}
 
-	static void ModifyCompilationEnvironment(const FGlobalShaderPermutationParameters& Parameters, FShaderCompilerEnvironment & OutEnvironment)
+	static void ModifyCompilationEnvironment(const FGlobalShaderPermutationParameters& Parameters, FShaderCompilerEnvironment& OutEnvironment)
 	{
 		FColorCorrectRegionsPostProcessMaterialShader::ModifyCompilationEnvironment(Parameters, OutEnvironment);
 	}
+
+};
+
+class FColorCorrectRegionMaterialPS : public FColorCorrectGenericPS
+{
+public:
+
+	DECLARE_GLOBAL_SHADER(FColorCorrectRegionMaterialPS);
+
+	class FShaderType : SHADER_PERMUTATION_ENUM_CLASS("SHAPE_TYPE", EColorCorrectRegionsType);
+
+	using FPermutationDomain = TShaderPermutationDomain<FShaderType, FTemperatureType, FAdvancedShader>;
+
+	FColorCorrectRegionMaterialPS() = default;
+	FColorCorrectRegionMaterialPS(const ShaderMetaType::CompiledShaderInitializerType & Initializer)
+		: FColorCorrectGenericPS(Initializer)
+	{}
+};
+
+
+class FColorCorrectWindowMaterialPS : public FColorCorrectGenericPS
+{
+public:
+
+	DECLARE_GLOBAL_SHADER(FColorCorrectWindowMaterialPS);
+
+	class FShaderType : SHADER_PERMUTATION_ENUM_CLASS("WINDOW_TYPE", EColorCorrectWindowType);
+
+	using FPermutationDomain = TShaderPermutationDomain<FShaderType, FTemperatureType, FAdvancedShader>;
+
+	FColorCorrectWindowMaterialPS() = default;
+	FColorCorrectWindowMaterialPS(const ShaderMetaType::CompiledShaderInitializerType& Initializer)
+		: FColorCorrectGenericPS(Initializer)
+	{}
 };
 
 // The vertex shader used by DrawScreenPass to draw a rectangle.
diff --git a/Engine/Plugins/Experimental/ColorCorrectRegions/Source/ColorCorrectRegions/Public/ColorCorrectRegionsSubsystem.h b/Engine/Plugins/Experimental/ColorCorrectRegions/Source/ColorCorrectRegions/Public/ColorCorrectRegionsSubsystem.h
--- a/Engine/Plugins/Experimental/ColorCorrectRegions/Source/ColorCorrectRegions/Public/ColorCorrectRegionsSubsystem.h
+++ b/Engine/Plugins/Experimental/ColorCorrectRegions/Source/ColorCorrectRegions/Public/ColorCorrectRegionsSubsystem.h
@@ -97,10 +97,6 @@ public:
 	TArray<AColorCorrectRegion*> RegionsDistanceBased;
 
 private:
-
-	/** Region class. Used for getting all region actors in level. */
-	TSubclassOf<AColorCorrectRegion> RegionClass;
-
 	TSharedPtr< class FColorCorrectRegionsSceneViewExtension, ESPMode::ThreadSafe > PostProcessSceneViewExtension;
 
 	FCriticalSection RegionAccessCriticalSection;
diff --git a/Engine/Plugins/Experimental/ColorCorrectRegions/Source/ColorCorrectRegions/Public/ColorCorrectWindow.h b/Engine/Plugins/Experimental/ColorCorrectRegions/Source/ColorCorrectRegions/Public/ColorCorrectWindow.h
new file mode 100644
--- /dev/null
+++ b/Engine/Plugins/Experimental/ColorCorrectRegions/Source/ColorCorrectRegions/Public/ColorCorrectWindow.h
@@ -0,0 +1,49 @@
+// Copyright Epic Games, Inc. All Rights Reserved.
+
+#pragma once
+
+#include "ColorCorrectRegion.h"
+#include "CoreMinimal.h"
+#include "UObject/ObjectMacros.h"
+#include "GameFramework/Actor.h"
+#include "Engine/Classes/Components/MeshComponent.h"
+#include "Engine/Scene.h"
+#include "ColorCorrectWindow.generated.h"
+
+
+class UColorCorrectRegionsSubsystem;
+class UBillboardComponent;
+
+UENUM(BlueprintType)
+enum class EColorCorrectWindowType : uint8
+{
+	Plane		UMETA(DisplayName = "Plane"),
+	Circle		UMETA(DisplayName = "Circle"),
+	MAX
+};
+
+/**
+ * 
+ */
+UCLASS(Blueprintable, notplaceable)
+class COLORCORRECTREGIONS_API AColorCorrectWindow : public AColorCorrectRegion
+{
+	GENERATED_UCLASS_BODY()
+public:
+
+	UPROPERTY()
+	TArray<TObjectPtr<UStaticMeshComponent>> MeshComponents;
+
+	/** Region type. */
+	UPROPERTY(BlueprintReadWrite, EditAnywhere, Category="Color Correction", Meta = (DisplayName = "Type"))
+	EColorCorrectWindowType WindowType;
+
+#if WITH_EDITOR
+	/** Called when any of the properties are changed. */
+	virtual void PostEditChangeProperty(FPropertyChangedEvent& PropertyChangedEvent) override;
+#endif
+
+private:
+	void SetMeshVisibilityForWindowType();
+
+};
diff --git a/Engine/Plugins/Experimental/ColorCorrectRegions/Source/ColorCorrectRegionsEditor/Private/ColorCorrectRegionCustomization.cpp b/Engine/Plugins/Experimental/ColorCorrectRegions/Source/ColorCorrectRegionsEditor/Private/ColorCorrectRegionCustomization.cpp
--- a/Engine/Plugins/Experimental/ColorCorrectRegions/Source/ColorCorrectRegionsEditor/Private/ColorCorrectRegionCustomization.cpp
+++ b/Engine/Plugins/Experimental/ColorCorrectRegions/Source/ColorCorrectRegionsEditor/Private/ColorCorrectRegionCustomization.cpp
@@ -2,41 +2,25 @@
 
 #include "ColorCorrectRegionCustomization.h"
 #include "ColorCorrectRegion.h"
+#include "ColorCorrectWindow.h"
 #include "DetailLayoutBuilder.h"
 
-TSharedRef<IDetailCustomization> FColorCorrectRegionDetails::MakeInstance()
+TSharedRef<IDetailCustomization> FColorCorrectWindowDetails::MakeInstance()
 {
-	return MakeShareable(new FColorCorrectRegionDetails);
+	return MakeShareable(new FColorCorrectWindowDetails);
 }
 
-void FColorCorrectRegionDetails::CustomizeDetails(IDetailLayoutBuilder& DetailBuilder)
+void FColorCorrectWindowDetails::CustomizeDetails(IDetailLayoutBuilder& DetailBuilder)
 {
-	bool bPriorityHidden = false;
 	for (const TWeakObjectPtr<UObject>& SelectedObject : DetailBuilder.GetSelectedObjects())
 	{
-		if (AColorCorrectRegion* CCR = Cast<AColorCorrectRegion>(SelectedObject.Get()))
+		if (AColorCorrectWindow* CCW = Cast<AColorCorrectWindow>(SelectedObject.Get()))
 		{
-			if (CCR->Type == EColorCorrectRegionsType::Plane)
-			{
-				bPriorityHidden = true;
-				break;
-			}
-		}
-	}
+			TSharedRef<IPropertyHandle> PriorityProperty = DetailBuilder.GetProperty(GET_MEMBER_NAME_CHECKED(AColorCorrectRegion, Priority));
+			DetailBuilder.HideProperty(PriorityProperty);
 
-	if (bPriorityHidden)
-	{
-		TSharedRef<IPropertyHandle> PriorityProperty = DetailBuilder.GetProperty(GET_MEMBER_NAME_CHECKED(AColorCorrectRegion, Priority));
-		DetailBuilder.HideProperty(PriorityProperty);
-	}
-
-	// Trick to force the viewport to update because depending on Type, Priority property could be hidden.
-	TSharedRef<IPropertyHandle> TypeProperty = DetailBuilder.GetProperty(GET_MEMBER_NAME_CHECKED(AColorCorrectRegion, Type));
-	if (TypeProperty->IsValidHandle())
-	{
-		TypeProperty->SetOnPropertyValueChanged(FSimpleDelegate::CreateLambda([&DetailBuilder]()
-			{
-				DetailBuilder.ForceRefreshDetails();
-			}));
+			TSharedRef<IPropertyHandle> TypeProperty = DetailBuilder.GetProperty(GET_MEMBER_NAME_CHECKED(AColorCorrectRegion, Type));
+			DetailBuilder.HideProperty(TypeProperty);
+		}
 	}
 }
diff --git a/Engine/Plugins/Experimental/ColorCorrectRegions/Source/ColorCorrectRegionsEditor/Private/ColorCorrectRegionCustomization.h b/Engine/Plugins/Experimental/ColorCorrectRegions/Source/ColorCorrectRegionsEditor/Private/ColorCorrectRegionCustomization.h
--- a/Engine/Plugins/Experimental/ColorCorrectRegions/Source/ColorCorrectRegionsEditor/Private/ColorCorrectRegionCustomization.h
+++ b/Engine/Plugins/Experimental/ColorCorrectRegions/Source/ColorCorrectRegionsEditor/Private/ColorCorrectRegionCustomization.h
@@ -8,7 +8,7 @@
 #include "PropertyHandle.h"
 
 
-class FColorCorrectRegionDetails : public IDetailCustomization
+class FColorCorrectWindowDetails : public IDetailCustomization
 {
 public:
 	/** Makes a new instance of this detail layout class for a specific detail view requesting it */
diff --git a/Engine/Plugins/Experimental/ColorCorrectRegions/Source/ColorCorrectRegionsEditor/Private/ColorCorrectRegionsEditorModule.cpp b/Engine/Plugins/Experimental/ColorCorrectRegions/Source/ColorCorrectRegionsEditor/Private/ColorCorrectRegionsEditorModule.cpp
--- a/Engine/Plugins/Experimental/ColorCorrectRegions/Source/ColorCorrectRegionsEditor/Private/ColorCorrectRegionsEditorModule.cpp
+++ b/Engine/Plugins/Experimental/ColorCorrectRegions/Source/ColorCorrectRegionsEditor/Private/ColorCorrectRegionsEditorModule.cpp
@@ -3,6 +3,7 @@
 #include "ColorCorrectRegionsEditorModule.h"
 #include "Editor/PlacementMode/Public/IPlacementModeModule.h"
 #include "ColorCorrectRegion.h"
+#include "ColorCorrectWindow.h"
 #include "ColorCorrectRegionsStyle.h"
 #include "ColorCorrectRegionCustomization.h"
 #include "ActorFactories/ActorFactoryBlueprint.h"
@@ -11,11 +12,11 @@
 
 void FColorCorrectRegionsEditorModule::StartupModule()
 {
-	IPlacementModeModule::Get().OnPlacementModeCategoryRefreshed().AddRaw(this, &FColorCorrectRegionsEditorModule::OnPlacementModeRefresh);
 	FColorCorrectRegionsStyle::Initialize();
+	IPlacementModeModule::Get().OnPlacementModeCategoryRefreshed().AddRaw(this, &FColorCorrectRegionsEditorModule::OnPlacementModeRefresh);
 
 	FPropertyEditorModule& PropertyModule = FModuleManager::LoadModuleChecked<FPropertyEditorModule>("PropertyEditor");
-	PropertyModule.RegisterCustomClassLayout(AColorCorrectRegion::StaticClass()->GetFName(), FOnGetDetailCustomizationInstance::CreateStatic(&FColorCorrectRegionDetails::MakeInstance));
+	PropertyModule.RegisterCustomClassLayout(AColorCorrectRegion::StaticClass()->GetFName(), FOnGetDetailCustomizationInstance::CreateStatic(&FColorCorrectWindowDetails::MakeInstance));
 }
 
 void FColorCorrectRegionsEditorModule::OnPlacementModeRefresh(FName CategoryName)
@@ -38,6 +39,17 @@ void FColorCorrectRegionsEditorModule::OnPlacementModeRefresh(FName CategoryName
 			NSLOCTEXT("PlacementMode", "Color Correct Region", "Color Correct Region")
 		);
 
+
+		FPlaceableItem* CCWPlaceableItem = new FPlaceableItem(
+			*AColorCorrectWindow::StaticClass(),
+			FAssetData(AColorCorrectWindow::StaticClass()),
+			FName("CCW.PlaceActorThumbnail"),
+			FName("CCR.PlaceActorIcon"),
+			TOptional<FLinearColor>(),
+			TOptional<int32>(), 
+			NSLOCTEXT("PlacementMode", "Color Correct Window", "Color Correct Window"));
+
+		PlacementModeModule.RegisterPlaceableItem(CategoryName, MakeShareable(CCWPlaceableItem));
 		PlacementModeModule.RegisterPlaceableItem(CategoryName, MakeShareable(CCRPlaceableItem));
 	}
 }
diff --git a/Engine/Plugins/Experimental/ColorCorrectRegions/Source/ColorCorrectRegionsEditor/Private/ColorCorrectRegionsStyle.cpp b/Engine/Plugins/Experimental/ColorCorrectRegions/Source/ColorCorrectRegionsEditor/Private/ColorCorrectRegionsStyle.cpp
--- a/Engine/Plugins/Experimental/ColorCorrectRegions/Source/ColorCorrectRegionsEditor/Private/ColorCorrectRegionsStyle.cpp
+++ b/Engine/Plugins/Experimental/ColorCorrectRegions/Source/ColorCorrectRegionsEditor/Private/ColorCorrectRegionsStyle.cpp
@@ -1,10 +1,10 @@
 // Copyright Epic Games, Inc. All Rights Reserved.
 
 #include "ColorCorrectRegionsStyle.h"
+#include "Interfaces/IPluginManager.h"
 #include "Styling/SlateStyle.h"
 #include "Styling/SlateStyleRegistry.h"
 #include "Styling/SlateTypes.h"
-#include "Interfaces/IPluginManager.h"
 
 TSharedPtr<FSlateStyleSet> FColorCorrectRegionsStyle::CCRStyle;
 
@@ -17,9 +17,11 @@ void FColorCorrectRegionsStyle::Initialize()
 
 	CCRStyle = MakeShared<FSlateStyleSet>(FName("ColorCorrectRegionsStyle"));
 
-	FString IconPath = IPluginManager::Get().FindPlugin(TEXT("ColorCorrectRegions"))->GetBaseDir() + TEXT("/Resources/PlaceActorPreview.png");
-	CCRStyle->Set("CCR.PlaceActorThumbnail", new FSlateImageBrush(IconPath, FVector2D(40.0f, 40.0f)));
-	CCRStyle->Set("CCR.PlaceActorIcon", new FSlateImageBrush(IconPath, FVector2D(16.0f, 16.0f)));
+	FString CCR_IconPath = IPluginManager::Get().FindPlugin(TEXT("ColorCorrectRegions"))->GetBaseDir() + TEXT("/Resources/PlaceActorPreview.png");
+	FString CCW_IconPath = IPluginManager::Get().FindPlugin(TEXT("ColorCorrectRegions"))->GetBaseDir() + TEXT("/Resources/Icon40_CCW.png");
+	CCRStyle->Set("CCR.PlaceActorThumbnail", new FSlateImageBrush(CCR_IconPath, FVector2D(40.0f, 40.0f)));
+	CCRStyle->Set("CCW.PlaceActorThumbnail", new FSlateImageBrush(CCW_IconPath, FVector2D(40.0f, 40.0f)));
+	CCRStyle->Set("CCR.PlaceActorIcon", new FSlateImageBrush(CCR_IconPath, FVector2D(16.0f, 16.0f)));
 
 	FSlateStyleRegistry::RegisterSlateStyle(*CCRStyle.Get());
 }
