Color Correct Regions: Added two options for stenciling out selected stencil IDs and only correcting selected stencil IDs. Part 1.
- Changed naming of CCW types.

diff --git a/Engine/Plugins/Experimental/ColorCorrectRegions/Shaders/Private/ColorCorrectRegionsShader.usf b/Engine/Plugins/Experimental/ColorCorrectRegions/Shaders/Private/ColorCorrectRegionsShader.usf
--- a/Engine/Plugins/Experimental/ColorCorrectRegions/Shaders/Private/ColorCorrectRegionsShader.usf
+++ b/Engine/Plugins/Experimental/ColorCorrectRegions/Shaders/Private/ColorCorrectRegionsShader.usf
@@ -32,6 +32,9 @@
 #define CCR_SHADER_DISPLAY_BOUNDING_RECT 0
 #define CLIP_PIXELS_OUTSIDE_AABB 1
 
+#define EXCLUDE_STENCIL 1
+
+
 struct FExternalExpressions
 {
 	float3 Rotate;
@@ -76,7 +79,8 @@ struct FExternalExpressions
 	float Intensity;
 
 	float FakeLight;
-	float ExcludeStencil;
+	int   ExcludeStencil;
+	int	  StencilId;
 
 	float4 SceneTexture;
 };
@@ -246,6 +250,41 @@ float GetShapeAlpha(in out FMaterialPixelParameters Parameters, in FExternalExpr
 
 MaterialFloat4 ProcessRegion(in out FMaterialPixelParameters Parameters, in FExternalExpressions ExternalExpressions)
 {
+	if (ExternalExpressions.ExcludeStencil > 0)
+	{
+		float CustomStencil = SceneTextureLookup(GetDefaultSceneTextureUV(Parameters, PPI_CustomStencil), PPI_CustomStencil, false).r;
+
+		const float Float_Eps = 0.00000001;
+
+		if (ExternalExpressions.ExcludeStencil == EXCLUDE_STENCIL)
+		{
+			if (CustomStencil == ExternalExpressions.StencilId)
+			{
+				if (abs(Parameters.SceneDepth - Parameters.CustomDepth) <= Float_Eps)
+				{
+					clip(-1.);
+					return MaterialFloat4(ExternalExpressions.SceneTexture.rgb, 1.);
+				}
+			}
+		}
+		else
+		{
+			if (CustomStencil == ExternalExpressions.StencilId)
+			{
+				if (abs(Parameters.SceneDepth - Parameters.CustomDepth) > Float_Eps)
+				{
+					clip(-1.);
+					return MaterialFloat4(ExternalExpressions.SceneTexture.rgb, 1.);
+				}
+			}
+			else
+			{
+				clip(-1.);
+				return MaterialFloat4(ExternalExpressions.SceneTexture.rgb, 1.);
+			}
+		}
+	}
+
 	float ShapeAlpha = GetShapeAlpha(Parameters, ExternalExpressions);
 
 	// Final check to reduce the performance based on shape.
@@ -270,15 +309,6 @@ MaterialFloat4 ProcessRegion(in out FMaterialPixelParameters Parameters, in FExt
 
 	MaterialFloat3 FinalColor = lerp(ExternalExpressions.SceneTexture.rgb, AfterTemperature, FinalAlpha);
 
-	if (ExternalExpressions.ExcludeStencil > 0)
-	{
-		float CustomStencil = SceneTextureLookup(GetDefaultSceneTextureUV(Parameters, PPI_CustomStencil), PPI_CustomStencil, false).r;
-
-		const float Float_Eps = 0.00000001;
-		float UseStencil = lerp(0, clamp(CustomStencil, 0., 1.), step(abs(Parameters.SceneDepth - Parameters.CustomDepth), Float_Eps));
-		FinalColor = lerp(FinalColor, ExternalExpressions.SceneTexture.rgb, UseStencil);
-	}
-
 	return MaterialFloat4(FinalColor, FinalAlpha);
 }
 
@@ -319,6 +349,7 @@ void CalcPixelMaterialInputs(in out FMaterialPixelParameters Parameters, in out
 		ExtExpressions.FakeLight = RegionData.FakeLight;
 		ExtExpressions.ExcludeStencil = RegionData.ExcludeStencil;
 		ExtExpressions.SceneTexture = SceneTexture.rgba;
+		ExtExpressions.StencilId = RegionData.StencilId;
 #if ADVANCED_CC	
 		{
 			ExtExpressions.ColorSaturationShadows = ColorCorrectShadows.ColorSaturation;
diff --git a/Engine/Plugins/Experimental/ColorCorrectRegions/Source/ColorCorrectRegions/Private/ColorCorrectRegion.cpp b/Engine/Plugins/Experimental/ColorCorrectRegions/Source/ColorCorrectRegions/Private/ColorCorrectRegion.cpp
--- a/Engine/Plugins/Experimental/ColorCorrectRegions/Source/ColorCorrectRegions/Private/ColorCorrectRegion.cpp
+++ b/Engine/Plugins/Experimental/ColorCorrectRegions/Source/ColorCorrectRegions/Private/ColorCorrectRegion.cpp
@@ -22,7 +22,7 @@ AColorCorrectRegion::AColorCorrectRegion(const FObjectInitializer& ObjectInitial
 	, Temperature(6500)
 	, Tint(0)
 	, Enabled(true)
-	, ExcludeStencil(false)
+	, ExcludeStencil(EColorCorrectRegionStencilType::None)
 	, ColorCorrectRegionsSubsystem(nullptr)
 {
 	PrimaryActorTick.bCanEverTick = true;
diff --git a/Engine/Plugins/Experimental/ColorCorrectRegions/Source/ColorCorrectRegions/Private/ColorCorrectRegionsSceneViewExtension.cpp b/Engine/Plugins/Experimental/ColorCorrectRegions/Source/ColorCorrectRegions/Private/ColorCorrectRegionsSceneViewExtension.cpp
--- a/Engine/Plugins/Experimental/ColorCorrectRegions/Source/ColorCorrectRegions/Private/ColorCorrectRegionsSceneViewExtension.cpp
+++ b/Engine/Plugins/Experimental/ColorCorrectRegions/Source/ColorCorrectRegions/Private/ColorCorrectRegionsSceneViewExtension.cpp
@@ -411,8 +411,9 @@ namespace
 			}
 			RegionData.Falloff = Region->Falloff;
 			RegionData.Intensity = Region->Intensity;
-			RegionData.ExcludeStencil = Region->ExcludeStencil;
 			RegionData.Invert = Region->Invert;
+			RegionData.ExcludeStencil = static_cast<uint32>(Region->ExcludeStencil);
+			RegionData.StencilId = Region->StencilId;
 
 			CCBase.ColorSaturation = (FVector4f)Region->ColorGradingSettings.Global.Saturation;
 			CCBase.ColorContrast = (FVector4f)Region->ColorGradingSettings.Global.Contrast;
diff --git a/Engine/Plugins/Experimental/ColorCorrectRegions/Source/ColorCorrectRegions/Private/ColorCorrectWindow.cpp b/Engine/Plugins/Experimental/ColorCorrectRegions/Source/ColorCorrectRegions/Private/ColorCorrectWindow.cpp
--- a/Engine/Plugins/Experimental/ColorCorrectRegions/Source/ColorCorrectRegions/Private/ColorCorrectWindow.cpp
+++ b/Engine/Plugins/Experimental/ColorCorrectRegions/Source/ColorCorrectRegions/Private/ColorCorrectWindow.cpp
@@ -3,6 +3,7 @@
 #include "ColorCorrectWindow.h"
 #include "Components/StaticMeshComponent.h"
 #include "CoreMinimal.h"
+#include "Engine/StaticMesh.h"
 #include "Engine/Texture2D.h"
 #include "UObject/ConstructorHelpers.h"
 
@@ -10,7 +11,7 @@ ENUM_RANGE_BY_COUNT(EColorCorrectWindowType, EColorCorrectWindowType::MAX)
 
 AColorCorrectWindow::AColorCorrectWindow(const FObjectInitializer& ObjectInitializer)
 	: Super(ObjectInitializer)
-	, WindowType(EColorCorrectWindowType::Plane)
+	, WindowType(EColorCorrectWindowType::Square)
 {
 	UMaterial* Material = LoadObject<UMaterial>(NULL, TEXT("/ColorCorrectRegions/Materials/M_ColorCorrectRegionTransparentPreview.M_ColorCorrectRegionTransparentPreview"), NULL, LOAD_None, NULL);
 	const TArray<UStaticMesh*> StaticMeshes =
diff --git a/Engine/Plugins/Experimental/ColorCorrectRegions/Source/ColorCorrectRegions/Public/ColorCorrectRegion.h b/Engine/Plugins/Experimental/ColorCorrectRegions/Source/ColorCorrectRegions/Public/ColorCorrectRegion.h
--- a/Engine/Plugins/Experimental/ColorCorrectRegions/Source/ColorCorrectRegions/Public/ColorCorrectRegion.h
+++ b/Engine/Plugins/Experimental/ColorCorrectRegions/Source/ColorCorrectRegions/Public/ColorCorrectRegion.h
@@ -32,6 +32,15 @@ enum class EColorCorrectRegionTemperatureType : uint8
 	MAX
 };
 
+UENUM(BlueprintType)
+enum class EColorCorrectRegionStencilType : uint8
+{
+	None					UMETA(DisplayName = "No Stencil"),
+	ExcludeStencil			UMETA(DisplayName = "Exclude Selected Stencil"),
+	IncludeStencil			UMETA(DisplayName = "Affect Only Selected Stencil"),
+	MAX
+};
+
 /**
  * An instance of Color Correction Region. Used to aggregate all active regions.
  * This actor is aggregated by ColorCorrectRegionsSubsystem which handles:
@@ -100,7 +109,11 @@ public:
 
 	/** Enable stenciling. */
 	UPROPERTY(BlueprintReadWrite, EditAnywhere, Category = "Color Correction")
-	bool ExcludeStencil;
+	EColorCorrectRegionStencilType ExcludeStencil;
+
+	/** Enable stenciling outside of stencil region. */
+	UPROPERTY(BlueprintReadWrite, EditAnywhere, Category = "Color Correction")
+	uint8 StencilId;
 
 #if WITH_EDITORONLY_DATA
 
diff --git a/Engine/Plugins/Experimental/ColorCorrectRegions/Source/ColorCorrectRegions/Public/ColorCorrectRegionsPostProcessMaterial.h b/Engine/Plugins/Experimental/ColorCorrectRegions/Source/ColorCorrectRegions/Public/ColorCorrectRegionsPostProcessMaterial.h
--- a/Engine/Plugins/Experimental/ColorCorrectRegions/Source/ColorCorrectRegions/Public/ColorCorrectRegionsPostProcessMaterial.h
+++ b/Engine/Plugins/Experimental/ColorCorrectRegions/Source/ColorCorrectRegions/Public/ColorCorrectRegionsPostProcessMaterial.h
@@ -34,7 +34,8 @@ BEGIN_GLOBAL_SHADER_PARAMETER_STRUCT(FCCRRegionDataInputParameter, )
 	SHADER_PARAMETER(float, Falloff)
 	SHADER_PARAMETER(float, Intensity)
 	SHADER_PARAMETER(float, FakeLight)
-	SHADER_PARAMETER(float, ExcludeStencil)
+	SHADER_PARAMETER(uint32, ExcludeStencil)
+	SHADER_PARAMETER(uint32, StencilId)
 	SHADER_PARAMETER(float, Invert)
 END_GLOBAL_SHADER_PARAMETER_STRUCT()
 
diff --git a/Engine/Plugins/Experimental/ColorCorrectRegions/Source/ColorCorrectRegions/Public/ColorCorrectWindow.h b/Engine/Plugins/Experimental/ColorCorrectRegions/Source/ColorCorrectRegions/Public/ColorCorrectWindow.h
--- a/Engine/Plugins/Experimental/ColorCorrectRegions/Source/ColorCorrectRegions/Public/ColorCorrectWindow.h
+++ b/Engine/Plugins/Experimental/ColorCorrectRegions/Source/ColorCorrectRegions/Public/ColorCorrectWindow.h
@@ -17,7 +17,7 @@ class UBillboardComponent;
 UENUM(BlueprintType)
 enum class EColorCorrectWindowType : uint8
 {
-	Plane		UMETA(DisplayName = "Plane"),
+	Square		UMETA(DisplayName = "Square"),
 	Circle		UMETA(DisplayName = "Circle"),
 	MAX
 };
