Color Correct Regions: Fixed PCH comiple issue.

diff --git a/Engine/Plugins/Experimental/ColorCorrectRegions/Source/ColorCorrectRegions/Private/ColorCorrectWindow.cpp b/Engine/Plugins/Experimental/ColorCorrectRegions/Source/ColorCorrectRegions/Private/ColorCorrectWindow.cpp
--- a/Engine/Plugins/Experimental/ColorCorrectRegions/Source/ColorCorrectRegions/Private/ColorCorrectWindow.cpp
+++ b/Engine/Plugins/Experimental/ColorCorrectRegions/Source/ColorCorrectRegions/Private/ColorCorrectWindow.cpp
@@ -1,8 +1,10 @@
 // Copyright Epic Games, Inc. All Rights Reserved.
 
 #include "ColorCorrectWindow.h"
-#include "Components/StaticMeshComponent.h"
 #include "CoreMinimal.h"
+#include "Materials/Material.h"
+#include "Engine/StaticMesh.h"
+#include "Components/StaticMeshComponent.h"
 #include "Engine/StaticMesh.h"
 #include "Engine/Texture2D.h"
 #include "UObject/ConstructorHelpers.h"
