Color Correct Regions:
- added ability to select multiple stencil numbers.
- added stencil merger that handles an array of stencils and provides to CCR and CCW for modifying only the selected regions.

diff --git a/Engine/Plugins/Experimental/ColorCorrectRegions/Shaders/Private/ColorCorrectRegionsCommon.usf b/Engine/Plugins/Experimental/ColorCorrectRegions/Shaders/Private/ColorCorrectRegionsCommon.usf
--- a/Engine/Plugins/Experimental/ColorCorrectRegions/Shaders/Private/ColorCorrectRegionsCommon.usf
+++ b/Engine/Plugins/Experimental/ColorCorrectRegions/Shaders/Private/ColorCorrectRegionsCommon.usf
@@ -101,18 +101,11 @@ struct FMaterialPixelParameters
 	 */
 	MaterialFloat4 SvPosition;
 
-	/* Scene depth in world units. */
-	half SceneDepth;
-
-	/* Custom depth in world units. */
-	half CustomDepth;
-
 	/** Post projection position reconstructed from SvPosition, before the divide by W. left..top -1..1, bottom..top -1..1  within the viewport, W is the SceneDepth */
 	MaterialFloat4 ScreenPosition;
 
 	/**
 	 * Interpolated worldspace position of this pixel
-	 * todo: Make this TranslatedWorldPosition and also rename the VS WorldPosition to be TranslatedWorldPosition
 	 */
 	MaterialFloat3 AbsoluteWorldPosition;
 
diff --git a/Engine/Plugins/Experimental/ColorCorrectRegions/Shaders/Private/ColorCorrectRegionsShader.usf b/Engine/Plugins/Experimental/ColorCorrectRegions/Shaders/Private/ColorCorrectRegionsShader.usf
--- a/Engine/Plugins/Experimental/ColorCorrectRegions/Shaders/Private/ColorCorrectRegionsShader.usf
+++ b/Engine/Plugins/Experimental/ColorCorrectRegions/Shaders/Private/ColorCorrectRegionsShader.usf
@@ -68,7 +68,6 @@ struct FExternalExpressions
 	float4 ColorOffsetHighlights;
 	float HighlightsMin;
 #endif
-
 	float WhiteTemp;
 	float Tint;
 	float Inner;
@@ -80,7 +79,6 @@ struct FExternalExpressions
 
 	float FakeLight;
 	int   ExcludeStencil;
-	int	  StencilId;
 
 	float4 SceneTexture;
 };
@@ -250,40 +248,28 @@ float GetShapeAlpha(in out FMaterialPixelParameters Parameters, in FExternalExpr
 
 MaterialFloat4 ProcessRegion(in out FMaterialPixelParameters Parameters, in FExternalExpressions ExternalExpressions)
 {
-	if (ExternalExpressions.ExcludeStencil > 0)
-	{
-		float CustomStencil = SceneTextureLookup(GetDefaultSceneTextureUV(Parameters, PPI_CustomStencil), PPI_CustomStencil, false).r;
-
-		const float Float_Eps = 0.00000001;
+#if STENCIL_ENABLED
+	int CustomStencil = CCRMergedStencil.MergedStencilTexture.Load(int3(Parameters.SvPosition.xy, 0.)).r;
 
+	{
 		if (ExternalExpressions.ExcludeStencil == EXCLUDE_STENCIL)
 		{
-			if (CustomStencil == ExternalExpressions.StencilId)
+			if (CustomStencil > 0)
 			{
-				if (abs(Parameters.SceneDepth - Parameters.CustomDepth) <= Float_Eps)
-				{
-					clip(-1.);
-					return MaterialFloat4(ExternalExpressions.SceneTexture.rgb, 1.);
-				}
+				clip(-1.);
+				return MaterialFloat4(ExternalExpressions.SceneTexture.rgb, 1.);
 			}
 		}
 		else
 		{
-			if (CustomStencil == ExternalExpressions.StencilId)
-			{
-				if (abs(Parameters.SceneDepth - Parameters.CustomDepth) > Float_Eps)
-				{
-					clip(-1.);
-					return MaterialFloat4(ExternalExpressions.SceneTexture.rgb, 1.);
-				}
-			}
-			else
+			if (CustomStencil == 0)
 			{
 				clip(-1.);
 				return MaterialFloat4(ExternalExpressions.SceneTexture.rgb, 1.);
 			}
 		}
 	}
+#endif
 
 	float ShapeAlpha = GetShapeAlpha(Parameters, ExternalExpressions);
 
@@ -349,7 +335,6 @@ void CalcPixelMaterialInputs(in out FMaterialPixelParameters Parameters, in out
 		ExtExpressions.FakeLight = RegionData.FakeLight;
 		ExtExpressions.ExcludeStencil = RegionData.ExcludeStencil;
 		ExtExpressions.SceneTexture = SceneTexture.rgba;
-		ExtExpressions.StencilId = RegionData.StencilId;
 #if ADVANCED_CC	
 		{
 			ExtExpressions.ColorSaturationShadows = ColorCorrectShadows.ColorSaturation;
@@ -435,8 +420,6 @@ void BuildParameters(in FPostProcessMaterialVSToPS Input, inout FMaterialPixelPa
 	Parameters.AbsoluteWorldPosition = TranslatedWorldPosition.xyz - LWCHackToFloat(ResolvedView.PreViewTranslation).xyz;
 
 	Parameters.SvPosition = SvPosition;
-	Parameters.SceneDepth = CalcSceneDepth(ViewportUVToBufferUV(ViewportUV));
-	Parameters.CustomDepth = CalcSceneCustomDepth(ViewportUVToBufferUV(ViewportUV));
 	Parameters.ScreenPosition = ScreenPosition;
 
 	Parameters.CameraVector = normalize(-Parameters.WorldPosition_CamRelative.xyz);
diff --git a/Engine/Plugins/Experimental/ColorCorrectRegions/Shaders/Private/ColorCorrectRegionsStencilMerger.usf b/Engine/Plugins/Experimental/ColorCorrectRegions/Shaders/Private/ColorCorrectRegionsStencilMerger.usf
new file mode 100644
--- /dev/null
+++ b/Engine/Plugins/Experimental/ColorCorrectRegions/Shaders/Private/ColorCorrectRegionsStencilMerger.usf
@@ -0,0 +1,36 @@
+// Copyright Epic Games, Inc. All Rights Reserved.
+
+#include "/Engine/Private/Common.ush"
+#include "/Engine/Private/DeferredShadingCommon.ush"
+#include "/Engine/Private/SceneTexturesCommon.ush"
+#include "/Engine/Private/ScreenPass.ush"
+#include "/Engine/Private/SceneTextureParameters.ush"
+#include "/Engine/Private/SceneData.ush"
+#include "/ColorCorrectRegionsShaders/Private/ColorCorrectRegionsCommon.usf"
+
+
+StructuredBuffer<uint> StencilIds;
+
+uint MainPS(
+	noperspective float4 UVAndScreenPos : TEXCOORD0
+) : SV_Target0
+{
+	float4 SvPosition = UVAndScreenPos;
+	const float Float_Eps = 0.00000000001;
+
+	half CustomDepth = CalcSceneCustomDepth(UVAndScreenPos.xy);
+	half SceneDepth = CalcSceneDepth(UVAndScreenPos.xy);
+
+	int CustomStencil = SceneTextureLookup(UVAndScreenPos.xy, PPI_CustomStencil, false).r;
+	for (uint Index = 0; Index < StencilIds.Length; Index++)
+	{
+		if (abs(SceneDepth - CustomDepth) <= Float_Eps)
+		{
+			if (CustomStencil == StencilIds[Index])
+			{
+				return CustomStencil;
+			}
+		}
+	}
+	return 0;
+}
diff --git a/Engine/Plugins/Experimental/ColorCorrectRegions/Source/ColorCorrectRegions/Private/ColorCorrectRegion.cpp b/Engine/Plugins/Experimental/ColorCorrectRegions/Source/ColorCorrectRegions/Private/ColorCorrectRegion.cpp
--- a/Engine/Plugins/Experimental/ColorCorrectRegions/Source/ColorCorrectRegions/Private/ColorCorrectRegion.cpp
+++ b/Engine/Plugins/Experimental/ColorCorrectRegions/Source/ColorCorrectRegions/Private/ColorCorrectRegion.cpp
@@ -22,7 +22,7 @@ AColorCorrectRegion::AColorCorrectRegion(const FObjectInitializer& ObjectInitial
 	, Temperature(6500)
 	, Tint(0)
 	, Enabled(true)
-	, ExcludeStencil(EColorCorrectRegionStencilType::None)
+	, PerActorColorCorrection(EColorCorrectRegionStencilType::None)
 	, ColorCorrectRegionsSubsystem(nullptr)
 {
 	PrimaryActorTick.bCanEverTick = true;
diff --git a/Engine/Plugins/Experimental/ColorCorrectRegions/Source/ColorCorrectRegions/Private/ColorCorrectRegionsPostProcessMaterial.cpp b/Engine/Plugins/Experimental/ColorCorrectRegions/Source/ColorCorrectRegions/Private/ColorCorrectRegionsPostProcessMaterial.cpp
--- a/Engine/Plugins/Experimental/ColorCorrectRegions/Source/ColorCorrectRegions/Private/ColorCorrectRegionsPostProcessMaterial.cpp
+++ b/Engine/Plugins/Experimental/ColorCorrectRegions/Source/ColorCorrectRegions/Private/ColorCorrectRegionsPostProcessMaterial.cpp
@@ -20,6 +20,7 @@ IMPLEMENT_GLOBAL_SHADER_PARAMETER_STRUCT(FCCRColorCorrectParameter, "ColorCorrec
 IMPLEMENT_GLOBAL_SHADER_PARAMETER_STRUCT(FCCRColorCorrectShadowsParameter, "ColorCorrectShadows");
 IMPLEMENT_GLOBAL_SHADER_PARAMETER_STRUCT(FCCRColorCorrectMidtonesParameter, "ColorCorrectMidtones");
 IMPLEMENT_GLOBAL_SHADER_PARAMETER_STRUCT(FCCRColorCorrectHighlightsParameter, "ColorCorrectHighlights");
+IMPLEMENT_GLOBAL_SHADER_PARAMETER_STRUCT(FCCRStencilMergerParameter, "CCRMergedStencil");
 
 
 IMPLEMENT_GLOBAL_SHADER(FColorCorrectRegionMaterialVS, "/ColorCorrectRegionsShaders/Private/ColorCorrectRegionsShader.usf", "MainVS", SF_Vertex);
@@ -27,3 +28,5 @@ IMPLEMENT_GLOBAL_SHADER(FColorCorrectRegionMaterialPS, "/ColorCorrectRegionsShad
 IMPLEMENT_GLOBAL_SHADER(FColorCorrectWindowMaterialPS, "/ColorCorrectRegionsShaders/Private/ColorCorrectRegionsShader.usf", "MainPS", SF_Pixel);
 IMPLEMENT_GLOBAL_SHADER(FColorCorrectScreenPassVS, "/ColorCorrectRegionsShaders/Private/ColorCorrectRegionsScreenPass.usf", "MainVS", SF_Vertex);
 IMPLEMENT_GLOBAL_SHADER(FClearRectPS, "/ColorCorrectRegionsShaders/Private/ColorCorrectRegionsScreenPass.usf", "MainPS", SF_Pixel);
+
+IMPLEMENT_GLOBAL_SHADER(FCCRStencilMergerPS, "/ColorCorrectRegionsShaders/Private/ColorCorrectRegionsStencilMerger.usf", "MainPS", SF_Pixel);
\ No newline at end of file
diff --git a/Engine/Plugins/Experimental/ColorCorrectRegions/Source/ColorCorrectRegions/Private/ColorCorrectRegionsSceneViewExtension.cpp b/Engine/Plugins/Experimental/ColorCorrectRegions/Source/ColorCorrectRegions/Private/ColorCorrectRegionsSceneViewExtension.cpp
--- a/Engine/Plugins/Experimental/ColorCorrectRegions/Source/ColorCorrectRegions/Private/ColorCorrectRegionsSceneViewExtension.cpp
+++ b/Engine/Plugins/Experimental/ColorCorrectRegions/Source/ColorCorrectRegions/Private/ColorCorrectRegionsSceneViewExtension.cpp
@@ -203,10 +203,11 @@ namespace
 	}
 
 	// A helper function for getting the right shader for SDF based CCRs.
-	TShaderMapRef<FColorCorrectRegionMaterialPS> GetRegionShader(const FGlobalShaderMap* GlobalShaderMap, EColorCorrectRegionsType RegionType, FColorCorrectGenericPS::ETemperatureType TemperatureType, bool bIsAdvanced)
+	TShaderMapRef<FColorCorrectRegionMaterialPS> GetRegionShader(const FGlobalShaderMap* GlobalShaderMap, EColorCorrectRegionsType RegionType, FColorCorrectGenericPS::ETemperatureType TemperatureType, bool bIsAdvanced, bool bUseStencil)
 	{
 		FColorCorrectRegionMaterialPS::FPermutationDomain PermutationVector;
 		PermutationVector.Set<FColorCorrectGenericPS::FAdvancedShader>(bIsAdvanced);
+		PermutationVector.Set<FColorCorrectGenericPS::FStencilEnabled>(bUseStencil);
 		PermutationVector.Set<FColorCorrectGenericPS::FTemperatureType>(TemperatureType);
 		PermutationVector.Set<FColorCorrectRegionMaterialPS::FShaderType>(static_cast<EColorCorrectRegionsType>(FMath::Min(static_cast<int32>(RegionType), static_cast<int32>(EColorCorrectRegionsType::MAX) - 1)));
 
@@ -214,10 +215,11 @@ namespace
 	}
 
 	// A helper function for getting the right shader for distance based CCRs.
-	TShaderMapRef<FColorCorrectWindowMaterialPS> GetWindowShader(const FGlobalShaderMap* GlobalShaderMap, EColorCorrectWindowType RegionType, FColorCorrectGenericPS::ETemperatureType TemperatureType, bool bIsAdvanced)
+	TShaderMapRef<FColorCorrectWindowMaterialPS> GetWindowShader(const FGlobalShaderMap* GlobalShaderMap, EColorCorrectWindowType RegionType, FColorCorrectGenericPS::ETemperatureType TemperatureType, bool bIsAdvanced, bool bUseStencil)
 	{
 		FColorCorrectWindowMaterialPS::FPermutationDomain PermutationVector;
 		PermutationVector.Set<FColorCorrectGenericPS::FAdvancedShader>(bIsAdvanced);
+		PermutationVector.Set<FColorCorrectGenericPS::FStencilEnabled>(bUseStencil);
 		PermutationVector.Set<FColorCorrectGenericPS::FTemperatureType>(TemperatureType);
 		PermutationVector.Set<FColorCorrectWindowMaterialPS::FShaderType>(static_cast<EColorCorrectWindowType>(FMath::Min(static_cast<int32>(RegionType), static_cast<int32>(EColorCorrectWindowType::MAX) - 1)));
 
@@ -233,6 +235,88 @@ namespace
 	}
 
 
+	void StencilMerger
+		( FRDGBuilder& GraphBuilder
+		, const AColorCorrectRegion* Region
+		, const FGlobalShaderMap* GlobalShaderMap
+		, const FScreenPassRenderTarget& SceneColorRenderTarget
+		, const FSceneView& View
+		, const FScreenPassTextureViewportParameters& SceneTextureViewportParams
+		, const FScreenPassTextureViewport& RegionViewport
+		, const FSceneTextureShaderParameters& SceneTextures
+		, const TArray<uint8>& StencilIds
+		, FScreenPassRenderTarget& OutMergedStencilRenderTarget)
+	{
+		if (StencilIds.Num() == 0)
+		{
+			return;
+		}
+		FRDGTextureDesc DepthBufferOutputDesc = SceneColorRenderTarget.Texture->Desc;
+		DepthBufferOutputDesc.Format = EPixelFormat::PF_DepthStencil;
+		DepthBufferOutputDesc.ClearValue = FClearValueBinding(0);
+		DepthBufferOutputDesc.Flags = TexCreate_RenderTargetable | TexCreate_DepthStencilTargetable;
+		DepthBufferOutputDesc.ClearValue = FClearValueBinding(0, 0);
+
+		FRDGTextureDesc Desc = SceneColorRenderTarget.Texture->Desc;
+		Desc.Format = EPixelFormat::PF_R8_UINT;
+		FRDGTexture* MergedStencilTexture = GraphBuilder.CreateTexture(Desc, TEXT("CCR_MergedStencil"));
+		OutMergedStencilRenderTarget = FScreenPassRenderTarget(MergedStencilTexture, SceneColorRenderTarget.ViewRect, ERenderTargetLoadAction::EClear);
+		{
+			TShaderMapRef<FCCRStencilMergerPS> StencilMergerPS(GlobalShaderMap);
+			TShaderMapRef<FColorCorrectScreenPassVS> StencilMergerVS(GlobalShaderMap);
+			FCCRStencilMergerPS::FParameters* Parameters = GraphBuilder.AllocParameters<FCCRStencilMergerPS::FParameters>();
+			Parameters->SceneTextures = SceneTextures;
+			Parameters->RenderTargets[0] = OutMergedStencilRenderTarget.GetRenderTargetBinding();
+			Parameters->PostProcessOutput = SceneTextureViewportParams;
+			Parameters->View = View.ViewUniformBuffer;
+
+			const FMatrix& RegionTransformMatrix = Region->GetActorTransform().ToMatrixWithScale();
+
+			FRHIBlendState* DefaultBlendState = FScreenPassPipelineState::FDefaultBlendState::GetRHI();
+				
+			FRHIResourceCreateInfo CreateInfo(TEXT("CCR_StencilIdBuffer"));
+
+			FBufferRHIRef BufferRef = RHICreateStructuredBuffer(sizeof(uint32), sizeof(uint32) * StencilIds.Num(), BUF_ShaderResource | BUF_Dynamic | BUF_FastVRAM, CreateInfo);
+			uint32* MappedBuffer = static_cast<uint32*>(RHILockBuffer(BufferRef, 0, StencilIds.Num(), RLM_WriteOnly));
+			uint8 Index = 0;
+			for (const uint8& StencilId : StencilIds)
+			{
+				MappedBuffer[Index++] = StencilId;
+			}
+
+			RHIUnlockBuffer(BufferRef);
+			Parameters->StencilIds = RHICreateShaderResourceView(BufferRef);
+
+			{
+				GraphBuilder.AddPass(
+					RDG_EVENT_NAME("ColorCorrectRegions_StencilMerger"),
+					Parameters,
+					ERDGPassFlags::Raster,
+						[&View,
+						StencilMergerVS,
+						StencilMergerPS,
+						Parameters,
+						RegionViewport,
+						DefaultBlendState](FRHICommandList& RHICmdList)
+					{
+						check(true);
+						DrawScreenPass(
+							RHICmdList,
+							static_cast<const FViewInfo&>(View),
+							RegionViewport,
+							RegionViewport,
+							FScreenPassPipelineState(StencilMergerVS, StencilMergerPS, DefaultBlendState, FScreenPassPipelineState::FDefaultDepthStencilState::GetRHI()),
+							EScreenPassDrawFlags::None,
+							[&](FRHICommandList& RHICmdList)
+							{
+								SetShaderParameters(RHICmdList, StencilMergerPS, StencilMergerPS.GetPixelShader(), *Parameters);
+							});
+					}
+				);
+			}
+		}
+	}
+
 	bool RenderRegion
 		( FRDGBuilder& GraphBuilder
 		, const FSceneView& View
@@ -372,14 +456,20 @@ namespace
 			? FColorCorrectRegionMaterialPS::ETemperatureType::Disabled
 			: static_cast<FColorCorrectRegionMaterialPS::ETemperatureType>(Region->TemperatureType);
 
+		FScreenPassRenderTarget MergedStencilRenderTarget;
+		if (Region->PerActorColorCorrection != EColorCorrectRegionStencilType::None)
+		{
+			StencilMerger(GraphBuilder, Region, GlobalShaderMap, SceneColorRenderTarget, View, SceneTextureViewportParams, RegionViewport, SceneTextures, Region->StencilIds, MergedStencilRenderTarget);
+		}
+
 		TShaderRef<FColorCorrectGenericPS> PixelShader;
 		if (AColorCorrectWindow* CCWindow = Cast<AColorCorrectWindow>(Region))
 		{
-			PixelShader = GetWindowShader(GlobalShaderMap, CCWindow->WindowType, TemperatureType, bIsAdvanced);
+			PixelShader = GetWindowShader(GlobalShaderMap, CCWindow->WindowType, TemperatureType, bIsAdvanced, MergedStencilRenderTarget.IsValid());
 		}
 		else
 		{
-			PixelShader = GetRegionShader(GlobalShaderMap, Region->Type, TemperatureType, bIsAdvanced);
+			PixelShader = GetRegionShader(GlobalShaderMap, Region->Type, TemperatureType, bIsAdvanced, MergedStencilRenderTarget.IsValid());
 		}
 
 		ClearUnusedGraphResources(VertexShader, PixelShader, PostProcessMaterialParameters);
@@ -389,8 +479,8 @@ namespace
 		FCCRColorCorrectShadowsParameter CCShadows;
 		FCCRColorCorrectMidtonesParameter CCMidtones;
 		FCCRColorCorrectHighlightsParameter CCHighlights;
-
-
+		FCCRStencilMergerParameter CCRMergedStencil;
+		
 		// Setting constant buffer data to be passed to the shader.
 		{
 			RegionData.Rotate = FMath::DegreesToRadians<FVector3f>((FVector3f)Region->GetActorRotation().Euler());
@@ -412,8 +502,7 @@ namespace
 			RegionData.Falloff = Region->Falloff;
 			RegionData.Intensity = Region->Intensity;
 			RegionData.Invert = Region->Invert;
-			RegionData.ExcludeStencil = static_cast<uint32>(Region->ExcludeStencil);
-			RegionData.StencilId = Region->StencilId;
+			RegionData.ExcludeStencil = static_cast<uint32>(Region->PerActorColorCorrection);
 
 			CCBase.ColorSaturation = (FVector4f)Region->ColorGradingSettings.Global.Saturation;
 			CCBase.ColorContrast = (FVector4f)Region->ColorGradingSettings.Global.Contrast;
@@ -447,6 +536,12 @@ namespace
 				CCHighlights.ColorOffset = (FVector4f)Region->ColorGradingSettings.Highlights.Offset;
 				CCHighlights.HighlightsMin = Region->ColorGradingSettings.HighlightsMin;
 			}
+
+			if (MergedStencilRenderTarget.IsValid())
+			{
+				CCRMergedStencil.MergedStencilTexture = MergedStencilRenderTarget.Texture;
+				CCRMergedStencil.MergedStencilSampler = TStaticSamplerState<>::GetRHI();
+			}
 		}
 
 #if CLIP_PIXELS_OUTSIDE_AABB
@@ -496,7 +591,9 @@ namespace
 			CCShadows,
 			CCMidtones,
 			CCHighlights,
-			bIsAdvanced](FRHICommandList& RHICmdList)
+			CCRMergedStencil,
+			bIsAdvanced,
+			MergedStencilRenderTarget](FRHICommandList& RHICmdList)
 			{
 				DrawScreenPass(
 					RHICmdList,
@@ -515,6 +612,10 @@ namespace
 							SetUniformBufferParameterImmediate(RHICmdList, PixelShader.GetPixelShader(), PixelShader->GetUniformBufferParameter<FCCRColorCorrectMidtonesParameter>(), CCMidtones);
 							SetUniformBufferParameterImmediate(RHICmdList, PixelShader.GetPixelShader(), PixelShader->GetUniformBufferParameter<FCCRColorCorrectHighlightsParameter>(), CCHighlights);
 						}
+						if (MergedStencilRenderTarget.IsValid())
+						{
+							SetUniformBufferParameterImmediate(RHICmdList, PixelShader.GetPixelShader(), PixelShader->GetUniformBufferParameter<FCCRStencilMergerParameter>(), CCRMergedStencil);
+						}
 						VertexShader->SetParameters(RHICmdList, View);
 						SetShaderParameters(RHICmdList, VertexShader, VertexShader.GetVertexShader(), *PostProcessMaterialParameters);
 
@@ -637,7 +738,6 @@ void FColorCorrectRegionsSceneViewExtension::PrePostProcessPass_RenderThread(FRD
 		check(View.bIsViewInfo);
 		FSceneTextureShaderParameters SceneTextures = CreateSceneTextureShaderParameters(GraphBuilder, ((const FViewInfo&)View).GetSceneTexturesChecked(), View.GetFeatureLevel(), ESceneTextureSetupMode::All);
 
-
 		for (auto It = WorldSubsystem->RegionsPriorityBased.CreateConstIterator(); It; ++It)
 		{
 			AColorCorrectRegion* Region = *It;
diff --git a/Engine/Plugins/Experimental/ColorCorrectRegions/Source/ColorCorrectRegions/Public/ColorCorrectRegion.h b/Engine/Plugins/Experimental/ColorCorrectRegions/Source/ColorCorrectRegions/Public/ColorCorrectRegion.h
--- a/Engine/Plugins/Experimental/ColorCorrectRegions/Source/ColorCorrectRegions/Public/ColorCorrectRegion.h
+++ b/Engine/Plugins/Experimental/ColorCorrectRegions/Source/ColorCorrectRegions/Public/ColorCorrectRegion.h
@@ -35,9 +35,9 @@ enum class EColorCorrectRegionTemperatureType : uint8
 UENUM(BlueprintType)
 enum class EColorCorrectRegionStencilType : uint8
 {
-	None					UMETA(DisplayName = "No Stencil"),
-	ExcludeStencil			UMETA(DisplayName = "Exclude Selected Stencil"),
-	IncludeStencil			UMETA(DisplayName = "Affect Only Selected Stencil"),
+	None					UMETA(DisplayName = "Default (Affect all actors)"),
+	ExcludeStencil			UMETA(DisplayName = "Exclude Selected Actors"),
+	IncludeStencil			UMETA(DisplayName = "Affect Only Selected Actors"),
 	MAX
 };
 
@@ -107,13 +107,13 @@ public:
 	UPROPERTY(BlueprintReadWrite, EditAnywhere, Category = "Color Correction")
 	bool Enabled;
 
-	/** Enable stenciling. */
+	/** Controls in which way the below targets will be affected by color correction. */
 	UPROPERTY(BlueprintReadWrite, EditAnywhere, Category = "Color Correction")
-	EColorCorrectRegionStencilType ExcludeStencil;
+	EColorCorrectRegionStencilType PerActorColorCorrection;
 
-	/** Enable stenciling outside of stencil region. */
+	/** A list of actors that need to be excluded or included from color correction (Depends on the option above). */
 	UPROPERTY(BlueprintReadWrite, EditAnywhere, Category = "Color Correction")
-	uint8 StencilId;
+	TArray<uint8> StencilIds;
 
 #if WITH_EDITORONLY_DATA
 
diff --git a/Engine/Plugins/Experimental/ColorCorrectRegions/Source/ColorCorrectRegions/Public/ColorCorrectRegionsPostProcessMaterial.h b/Engine/Plugins/Experimental/ColorCorrectRegions/Source/ColorCorrectRegions/Public/ColorCorrectRegionsPostProcessMaterial.h
--- a/Engine/Plugins/Experimental/ColorCorrectRegions/Source/ColorCorrectRegions/Public/ColorCorrectRegionsPostProcessMaterial.h
+++ b/Engine/Plugins/Experimental/ColorCorrectRegions/Source/ColorCorrectRegions/Public/ColorCorrectRegionsPostProcessMaterial.h
@@ -35,7 +35,6 @@ BEGIN_GLOBAL_SHADER_PARAMETER_STRUCT(FCCRRegionDataInputParameter, )
 	SHADER_PARAMETER(float, Intensity)
 	SHADER_PARAMETER(float, FakeLight)
 	SHADER_PARAMETER(uint32, ExcludeStencil)
-	SHADER_PARAMETER(uint32, StencilId)
 	SHADER_PARAMETER(float, Invert)
 END_GLOBAL_SHADER_PARAMETER_STRUCT()
 
@@ -77,6 +76,10 @@ BEGIN_GLOBAL_SHADER_PARAMETER_STRUCT(FCCRColorCorrectHighlightsParameter, )
 	SHADER_PARAMETER(float, HighlightsMin)
 END_GLOBAL_SHADER_PARAMETER_STRUCT()
 
+BEGIN_GLOBAL_SHADER_PARAMETER_STRUCT(FCCRStencilMergerParameter, )
+	SHADER_PARAMETER_RDG_TEXTURE(Texture2D<uint>, MergedStencilTexture)
+	SHADER_PARAMETER_SAMPLER(SamplerState, MergedStencilSampler)
+END_GLOBAL_SHADER_PARAMETER_STRUCT()
 
 BEGIN_SHADER_PARAMETER_STRUCT(FCCRShaderInputParameters, )
 	SHADER_PARAMETER_STRUCT_REF(FViewUniformShaderParameters, View)
@@ -138,6 +141,7 @@ public:
 
 	class FTemperatureType : SHADER_PERMUTATION_ENUM_CLASS("TEMPERATURE_TYPE", ETemperatureType);
 	class FAdvancedShader : SHADER_PERMUTATION_BOOL("ADVANCED_CC");
+	class FStencilEnabled : SHADER_PERMUTATION_BOOL("STENCIL_ENABLED");
 
 	FColorCorrectGenericPS() = default;
 	FColorCorrectGenericPS(const ShaderMetaType::CompiledShaderInitializerType& Initializer)
@@ -164,7 +168,7 @@ public:
 
 	class FShaderType : SHADER_PERMUTATION_ENUM_CLASS("SHAPE_TYPE", EColorCorrectRegionsType);
 
-	using FPermutationDomain = TShaderPermutationDomain<FShaderType, FTemperatureType, FAdvancedShader>;
+	using FPermutationDomain = TShaderPermutationDomain<FShaderType, FTemperatureType, FAdvancedShader, FStencilEnabled>;
 
 	FColorCorrectRegionMaterialPS() = default;
 	FColorCorrectRegionMaterialPS(const ShaderMetaType::CompiledShaderInitializerType & Initializer)
@@ -181,7 +185,7 @@ public:
 
 	class FShaderType : SHADER_PERMUTATION_ENUM_CLASS("WINDOW_TYPE", EColorCorrectWindowType);
 
-	using FPermutationDomain = TShaderPermutationDomain<FShaderType, FTemperatureType, FAdvancedShader>;
+	using FPermutationDomain = TShaderPermutationDomain<FShaderType, FTemperatureType, FAdvancedShader, FStencilEnabled>;
 
 	FColorCorrectWindowMaterialPS() = default;
 	FColorCorrectWindowMaterialPS(const ShaderMetaType::CompiledShaderInitializerType& Initializer)
@@ -227,3 +231,43 @@ class FClearRectPS : public FGlobalShader
 		RENDER_TARGET_BINDING_SLOTS()
 	END_SHADER_PARAMETER_STRUCT()
 };
+
+
+/** Vertex shader to rasterize obstructive CCR Plane. */
+class FCCRStencilMergerVS : public FGlobalShader
+{
+public:
+	DECLARE_GLOBAL_SHADER(FCCRStencilMergerVS);
+
+	static bool ShouldCompilePermutation(const FGlobalShaderPermutationParameters&)
+	{
+		return true;
+	}
+
+	FCCRStencilMergerVS() = default;
+	FCCRStencilMergerVS(const ShaderMetaType::CompiledShaderInitializerType& Initializer)
+		: FGlobalShader(Initializer)
+	{}
+};
+
+/** Pixel shader to rasterize obstructive CCR Plane. */
+class FCCRStencilMergerPS : public FGlobalShader
+{
+	DECLARE_GLOBAL_SHADER(FCCRStencilMergerPS);
+	SHADER_USE_PARAMETER_STRUCT(FCCRStencilMergerPS, FGlobalShader);
+
+	static bool ShouldCompilePermutation(const FGlobalShaderPermutationParameters& Parameters)
+	{
+		return true;
+	}
+
+	BEGIN_SHADER_PARAMETER_STRUCT(FParameters, )
+		//SHADER_PARAMETER_RDG_TEXTURE(Texture2D, InputTexture)
+		//SHADER_PARAMETER_SAMPLER(SamplerState, InputSampler)
+		SHADER_PARAMETER_STRUCT_REF(FViewUniformShaderParameters, View)
+		SHADER_PARAMETER_SRV(StructuredBuffer<uint>, StencilIds)
+		SHADER_PARAMETER_STRUCT_INCLUDE(FSceneTextureShaderParameters, SceneTextures)
+		SHADER_PARAMETER_STRUCT(FScreenPassTextureViewportParameters, PostProcessOutput)
+		RENDER_TARGET_BINDING_SLOTS()
+	END_SHADER_PARAMETER_STRUCT()
+};
\ No newline at end of file
