Color Correct Regions: Fixed issue with StructuredBuffer compatibility with various platforms.

diff --git a/Engine/Plugins/Experimental/ColorCorrectRegions/Shaders/Private/ColorCorrectRegionsStencilMerger.usf b/Engine/Plugins/Experimental/ColorCorrectRegions/Shaders/Private/ColorCorrectRegionsStencilMerger.usf
--- a/Engine/Plugins/Experimental/ColorCorrectRegions/Shaders/Private/ColorCorrectRegionsStencilMerger.usf
+++ b/Engine/Plugins/Experimental/ColorCorrectRegions/Shaders/Private/ColorCorrectRegionsStencilMerger.usf
@@ -10,6 +10,7 @@
 
 
 StructuredBuffer<uint> StencilIds;
+uint StencilIdCount;
 
 uint MainPS(
 	noperspective float4 UVAndScreenPos : TEXCOORD0
@@ -22,7 +23,7 @@ uint MainPS(
 	half SceneDepth = CalcSceneDepth(UVAndScreenPos.xy);
 
 	int CustomStencil = SceneTextureLookup(UVAndScreenPos.xy, PPI_CustomStencil, false).r;
-	for (uint Index = 0; Index < StencilIds.Length; Index++)
+	for (uint Index = 0; Index < StencilIdCount; Index++)
 	{
 		if (abs(SceneDepth - CustomDepth) <= Float_Eps)
 		{
diff --git a/Engine/Plugins/Experimental/ColorCorrectRegions/Source/ColorCorrectRegions/Private/ColorCorrectRegionsSceneViewExtension.cpp b/Engine/Plugins/Experimental/ColorCorrectRegions/Source/ColorCorrectRegions/Private/ColorCorrectRegionsSceneViewExtension.cpp
--- a/Engine/Plugins/Experimental/ColorCorrectRegions/Source/ColorCorrectRegions/Private/ColorCorrectRegionsSceneViewExtension.cpp
+++ b/Engine/Plugins/Experimental/ColorCorrectRegions/Source/ColorCorrectRegions/Private/ColorCorrectRegionsSceneViewExtension.cpp
@@ -286,6 +286,7 @@ namespace
 
 			RHIUnlockBuffer(BufferRef);
 			Parameters->StencilIds = RHICreateShaderResourceView(BufferRef);
+			Parameters->StencilIdCount = StencilIds.Num();
 
 			{
 				GraphBuilder.AddPass(
diff --git a/Engine/Plugins/Experimental/ColorCorrectRegions/Source/ColorCorrectRegions/Public/ColorCorrectRegionsPostProcessMaterial.h b/Engine/Plugins/Experimental/ColorCorrectRegions/Source/ColorCorrectRegions/Public/ColorCorrectRegionsPostProcessMaterial.h
--- a/Engine/Plugins/Experimental/ColorCorrectRegions/Source/ColorCorrectRegions/Public/ColorCorrectRegionsPostProcessMaterial.h
+++ b/Engine/Plugins/Experimental/ColorCorrectRegions/Source/ColorCorrectRegions/Public/ColorCorrectRegionsPostProcessMaterial.h
@@ -262,10 +262,9 @@ class FCCRStencilMergerPS : public FGlobalShader
 	}
 
 	BEGIN_SHADER_PARAMETER_STRUCT(FParameters, )
-		//SHADER_PARAMETER_RDG_TEXTURE(Texture2D, InputTexture)
-		//SHADER_PARAMETER_SAMPLER(SamplerState, InputSampler)
 		SHADER_PARAMETER_STRUCT_REF(FViewUniformShaderParameters, View)
 		SHADER_PARAMETER_SRV(StructuredBuffer<uint>, StencilIds)
+		SHADER_PARAMETER(uint32, StencilIdCount)
 		SHADER_PARAMETER_STRUCT_INCLUDE(FSceneTextureShaderParameters, SceneTextures)
 		SHADER_PARAMETER_STRUCT(FScreenPassTextureViewportParameters, PostProcessOutput)
 		RENDER_TARGET_BINDING_SLOTS()
