Color Correct Regions: All render resources are now async RDG compatible. Fixed a couple of crashes related to RDG and RHI parallel rendering in CCR.

diff --git a/Engine/Plugins/Experimental/ColorCorrectRegions/Shaders/Private/ColorCorrectRegionsShader.usf b/Engine/Plugins/Experimental/ColorCorrectRegions/Shaders/Private/ColorCorrectRegionsShader.usf
--- a/Engine/Plugins/Experimental/ColorCorrectRegions/Shaders/Private/ColorCorrectRegionsShader.usf
+++ b/Engine/Plugins/Experimental/ColorCorrectRegions/Shaders/Private/ColorCorrectRegionsShader.usf
@@ -34,6 +34,7 @@
 
 #define EXCLUDE_STENCIL 1
 
+Texture2D<uint> MergedStencilTexture;
 
 struct FExternalExpressions
 {
@@ -249,7 +250,7 @@ float GetShapeAlpha(in out FMaterialPixelParameters Parameters, in FExternalExpr
 MaterialFloat4 ProcessRegion(in out FMaterialPixelParameters Parameters, in FExternalExpressions ExternalExpressions)
 {
 #if STENCIL_ENABLED
-	int CustomStencil = CCRMergedStencil.MergedStencilTexture.Load(int3(Parameters.SvPosition.xy, 0.)).r;
+	int CustomStencil = MergedStencilTexture.Load(int3(Parameters.SvPosition.xy, 0.)).r;
 
 	{
 		if (ExternalExpressions.ExcludeStencil == EXCLUDE_STENCIL)
diff --git a/Engine/Plugins/Experimental/ColorCorrectRegions/Source/ColorCorrectRegions/Private/ColorCorrectRegionsPostProcessMaterial.cpp b/Engine/Plugins/Experimental/ColorCorrectRegions/Source/ColorCorrectRegions/Private/ColorCorrectRegionsPostProcessMaterial.cpp
--- a/Engine/Plugins/Experimental/ColorCorrectRegions/Source/ColorCorrectRegions/Private/ColorCorrectRegionsPostProcessMaterial.cpp
+++ b/Engine/Plugins/Experimental/ColorCorrectRegions/Source/ColorCorrectRegions/Private/ColorCorrectRegionsPostProcessMaterial.cpp
@@ -20,7 +20,6 @@ IMPLEMENT_GLOBAL_SHADER_PARAMETER_STRUCT(FCCRColorCorrectParameter, "ColorCorrec
 IMPLEMENT_GLOBAL_SHADER_PARAMETER_STRUCT(FCCRColorCorrectShadowsParameter, "ColorCorrectShadows");
 IMPLEMENT_GLOBAL_SHADER_PARAMETER_STRUCT(FCCRColorCorrectMidtonesParameter, "ColorCorrectMidtones");
 IMPLEMENT_GLOBAL_SHADER_PARAMETER_STRUCT(FCCRColorCorrectHighlightsParameter, "ColorCorrectHighlights");
-IMPLEMENT_GLOBAL_SHADER_PARAMETER_STRUCT(FCCRStencilMergerParameter, "CCRMergedStencil");
 
 
 IMPLEMENT_GLOBAL_SHADER(FColorCorrectRegionMaterialVS, "/ColorCorrectRegionsShaders/Private/ColorCorrectRegionsShader.usf", "MainVS", SF_Vertex);
diff --git a/Engine/Plugins/Experimental/ColorCorrectRegions/Source/ColorCorrectRegions/Private/ColorCorrectRegionsSceneViewExtension.cpp b/Engine/Plugins/Experimental/ColorCorrectRegions/Source/ColorCorrectRegions/Private/ColorCorrectRegionsSceneViewExtension.cpp
--- a/Engine/Plugins/Experimental/ColorCorrectRegions/Source/ColorCorrectRegions/Private/ColorCorrectRegionsSceneViewExtension.cpp
+++ b/Engine/Plugins/Experimental/ColorCorrectRegions/Source/ColorCorrectRegions/Private/ColorCorrectRegionsSceneViewExtension.cpp
@@ -276,16 +276,14 @@ namespace
 				
 			FRHIResourceCreateInfo CreateInfo(TEXT("CCR_StencilIdBuffer"));
 
-			FBufferRHIRef BufferRef = RHICreateStructuredBuffer(sizeof(uint32), sizeof(uint32) * StencilIds.Num(), BUF_ShaderResource | BUF_Dynamic | BUF_FastVRAM, CreateInfo);
-			uint32* MappedBuffer = static_cast<uint32*>(RHILockBuffer(BufferRef, 0, StencilIds.Num(), RLM_WriteOnly));
-			uint8 Index = 0;
-			for (const uint8& StencilId : StencilIds)
+			TArray<uint32> BufferData;
+			BufferData.AddUninitialized(StencilIds.Num());
+			for (int Index = 0; Index < BufferData.Num(); Index++)
 			{
-				MappedBuffer[Index++] = StencilId;
+				BufferData[Index] = StencilIds[Index];
 			}
 
-			RHIUnlockBuffer(BufferRef);
-			Parameters->StencilIds = RHICreateShaderResourceView(BufferRef);
+			Parameters->StencilIds = GraphBuilder.CreateSRV(FRDGBufferSRVDesc(CreateStructuredBuffer(GraphBuilder, TEXT("CCR.StencilIdBuffer"), sizeof(uint32), StencilIds.Num(), &BufferData[0], sizeof(uint32) * StencilIds.Num())));
 			Parameters->StencilIdCount = StencilIds.Num();
 
 			{
@@ -473,6 +471,11 @@ namespace
 			PixelShader = GetRegionShader(GlobalShaderMap, Region->Type, TemperatureType, bIsAdvanced, MergedStencilRenderTarget.IsValid());
 		}
 
+		if (MergedStencilRenderTarget.IsValid())
+		{
+			PostProcessMaterialParameters->MergedStencilTexture = MergedStencilRenderTarget.Texture;
+		}
+
 		ClearUnusedGraphResources(VertexShader, PixelShader, PostProcessMaterialParameters);
 
 		FCCRRegionDataInputParameter RegionData;
@@ -480,7 +483,6 @@ namespace
 		FCCRColorCorrectShadowsParameter CCShadows;
 		FCCRColorCorrectMidtonesParameter CCMidtones;
 		FCCRColorCorrectHighlightsParameter CCHighlights;
-		FCCRStencilMergerParameter CCRMergedStencil;
 		
 		// Setting constant buffer data to be passed to the shader.
 		{
@@ -538,11 +540,6 @@ namespace
 				CCHighlights.HighlightsMin = Region->ColorGradingSettings.HighlightsMin;
 			}
 
-			if (MergedStencilRenderTarget.IsValid())
-			{
-				CCRMergedStencil.MergedStencilTexture = MergedStencilRenderTarget.Texture;
-				CCRMergedStencil.MergedStencilSampler = TStaticSamplerState<>::GetRHI();
-			}
 		}
 
 #if CLIP_PIXELS_OUTSIDE_AABB
@@ -592,7 +589,6 @@ namespace
 			CCShadows,
 			CCMidtones,
 			CCHighlights,
-			CCRMergedStencil,
 			bIsAdvanced,
 			MergedStencilRenderTarget](FRHICommandList& RHICmdList)
 			{
@@ -613,10 +609,7 @@ namespace
 							SetUniformBufferParameterImmediate(RHICmdList, PixelShader.GetPixelShader(), PixelShader->GetUniformBufferParameter<FCCRColorCorrectMidtonesParameter>(), CCMidtones);
 							SetUniformBufferParameterImmediate(RHICmdList, PixelShader.GetPixelShader(), PixelShader->GetUniformBufferParameter<FCCRColorCorrectHighlightsParameter>(), CCHighlights);
 						}
-						if (MergedStencilRenderTarget.IsValid())
-						{
-							SetUniformBufferParameterImmediate(RHICmdList, PixelShader.GetPixelShader(), PixelShader->GetUniformBufferParameter<FCCRStencilMergerParameter>(), CCRMergedStencil);
-						}
+
 						VertexShader->SetParameters(RHICmdList, View);
 						SetShaderParameters(RHICmdList, VertexShader, VertexShader.GetVertexShader(), *PostProcessMaterialParameters);
 
diff --git a/Engine/Plugins/Experimental/ColorCorrectRegions/Source/ColorCorrectRegions/Public/ColorCorrectRegionsPostProcessMaterial.h b/Engine/Plugins/Experimental/ColorCorrectRegions/Source/ColorCorrectRegions/Public/ColorCorrectRegionsPostProcessMaterial.h
--- a/Engine/Plugins/Experimental/ColorCorrectRegions/Source/ColorCorrectRegions/Public/ColorCorrectRegionsPostProcessMaterial.h
+++ b/Engine/Plugins/Experimental/ColorCorrectRegions/Source/ColorCorrectRegions/Public/ColorCorrectRegionsPostProcessMaterial.h
@@ -76,12 +76,9 @@ BEGIN_GLOBAL_SHADER_PARAMETER_STRUCT(FCCRColorCorrectHighlightsParameter, )
 	SHADER_PARAMETER(float, HighlightsMin)
 END_GLOBAL_SHADER_PARAMETER_STRUCT()
 
-BEGIN_GLOBAL_SHADER_PARAMETER_STRUCT(FCCRStencilMergerParameter, )
-	SHADER_PARAMETER_RDG_TEXTURE(Texture2D<uint>, MergedStencilTexture)
-	SHADER_PARAMETER_SAMPLER(SamplerState, MergedStencilSampler)
-END_GLOBAL_SHADER_PARAMETER_STRUCT()
 
 BEGIN_SHADER_PARAMETER_STRUCT(FCCRShaderInputParameters, )
+	SHADER_PARAMETER_RDG_TEXTURE(Texture2D<uint>, MergedStencilTexture)
 	SHADER_PARAMETER_STRUCT_REF(FViewUniformShaderParameters, View)
 	SHADER_PARAMETER_STRUCT_INCLUDE(FSceneTextureShaderParameters, SceneTextures)
 	SHADER_PARAMETER_STRUCT_REF(FWorkingColorSpaceShaderParameters, WorkingColorSpace)
@@ -263,7 +260,7 @@ class FCCRStencilMergerPS : public FGlobalShader
 
 	BEGIN_SHADER_PARAMETER_STRUCT(FParameters, )
 		SHADER_PARAMETER_STRUCT_REF(FViewUniformShaderParameters, View)
-		SHADER_PARAMETER_SRV(StructuredBuffer<uint>, StencilIds)
+		SHADER_PARAMETER_RDG_BUFFER_SRV(StructuredBuffer<uint>, StencilIds)
 		SHADER_PARAMETER(uint32, StencilIdCount)
 		SHADER_PARAMETER_STRUCT_INCLUDE(FSceneTextureShaderParameters, SceneTextures)
 		SHADER_PARAMETER_STRUCT(FScreenPassTextureViewportParameters, PostProcessOutput)
