Color Correct Regions: Users are now able to select actors from the UI and CCR system will manage stencil Ids. This is a V1 and there are a few known corner cases that will be handled in the next CL.
- Added Stencil Manager that is endemic to Color Correct Regions.
- Changed default option for Per Actor CC mode
- added Per actor CC/enable disable option.

diff --git a/Engine/Plugins/Experimental/ColorCorrectRegions/Shaders/Private/ColorCorrectRegionsShader.usf b/Engine/Plugins/Experimental/ColorCorrectRegions/Shaders/Private/ColorCorrectRegionsShader.usf
--- a/Engine/Plugins/Experimental/ColorCorrectRegions/Shaders/Private/ColorCorrectRegionsShader.usf
+++ b/Engine/Plugins/Experimental/ColorCorrectRegions/Shaders/Private/ColorCorrectRegionsShader.usf
@@ -32,7 +32,7 @@
 #define CCR_SHADER_DISPLAY_BOUNDING_RECT 0
 #define CLIP_PIXELS_OUTSIDE_AABB 1
 
-#define EXCLUDE_STENCIL 1
+#define EXCLUDE_STENCIL 0
 
 Texture2D<uint> MergedStencilTexture;
 
@@ -79,7 +79,7 @@ struct FExternalExpressions
 	float Intensity;
 
 	float FakeLight;
-	int   ExcludeStencil;
+	uint   ExcludeStencil;
 
 	float4 SceneTexture;
 };
diff --git a/Engine/Plugins/Experimental/ColorCorrectRegions/Source/ColorCorrectRegions/Private/ColorCorrectRegion.cpp b/Engine/Plugins/Experimental/ColorCorrectRegions/Source/ColorCorrectRegions/Private/ColorCorrectRegion.cpp
--- a/Engine/Plugins/Experimental/ColorCorrectRegions/Source/ColorCorrectRegions/Private/ColorCorrectRegion.cpp
+++ b/Engine/Plugins/Experimental/ColorCorrectRegions/Source/ColorCorrectRegions/Private/ColorCorrectRegion.cpp
@@ -1,13 +1,15 @@
 // Copyright Epic Games, Inc. All Rights Reserved.
 
 #include "ColorCorrectRegion.h"
-#include "ColorCorrectRegionsSubsystem.h"
 #include "Components/BillboardComponent.h"
-#include "Components/StaticMeshComponent.h"
+#include "ColorCorrectRegionsSubsystem.h"
+#include "UObject/ConstructorHelpers.h"
 #include "CoreMinimal.h"
+#include "Engine/GameEngine.h"
 #include "Engine/Classes/Components/MeshComponent.h"
+#include "Components/StaticMeshComponent.h"
 #include "Engine/Texture2D.h"
-#include "UObject/ConstructorHelpers.h"
+
 
 AColorCorrectRegion::AColorCorrectRegion(const FObjectInitializer& ObjectInitializer) 
 	: Super(ObjectInitializer)
@@ -22,7 +24,8 @@ AColorCorrectRegion::AColorCorrectRegion(const FObjectInitializer& ObjectInitial
 	, Temperature(6500)
 	, Tint(0)
 	, Enabled(true)
-	, PerActorColorCorrection(EColorCorrectRegionStencilType::None)
+	, bEnablePerActorCC(false)
+	, PerActorColorCorrection(EColorCorrectRegionStencilType::IncludeStencil)
 	, ColorCorrectRegionsSubsystem(nullptr)
 {
 	PrimaryActorTick.bCanEverTick = true;
@@ -118,13 +121,17 @@ bool AColorCorrectRegion::ShouldTickIfViewportsOnly() const
 void AColorCorrectRegion::TickActor(float DeltaTime, ELevelTick TickType, FActorTickFunction& ThisTickFunction)
 {
 	Super::Tick(DeltaTime);
+
+	HandleAffectedActorsPropertyChange();
+	TransferStencilIds();
+
 	FTransform CurrentFrameTransform = GetTransform();
 	if (!PreviousFrameTransform.Equals(CurrentFrameTransform))
 	{
 		PreviousFrameTransform = CurrentFrameTransform;
 		GetActorBounds(true, BoxOrigin, BoxExtent);
 	}
-
+	
 	// Display Cluster uses HiddenPrimitives to hide Primitive components from view. 
 	// Store component id to be used on render thread.
 	if (const UStaticMeshComponent* FirstMeshComponent = FindComponentByClass<UStaticMeshComponent>())
@@ -141,6 +148,62 @@ void AColorCorrectRegion::Cleanup()
 	ColorCorrectRegionsSubsystem = nullptr;
 }
 
+void AColorCorrectRegion::TransferStencilIds()
+{
+	{
+		TArray<uint32> TempStencilIds;
+		for (TSoftObjectPtr<AActor> StencilActor : AffectedActors)
+		{
+			if (!StencilActor.IsValid())
+			{
+				continue;
+			}
+			TArray<UPrimitiveComponent*> PrimitiveComponents;
+			StencilActor->GetComponents<UPrimitiveComponent>(PrimitiveComponents);
+			for (UPrimitiveComponent* PrimitiveComponent : PrimitiveComponents)
+			{
+				if (PrimitiveComponent->bRenderCustomDepth)
+				{
+					TempStencilIds.Add(static_cast<uint32>(PrimitiveComponent->CustomDepthStencilValue));
+				}
+			}
+		}
+		{
+			FScopeLock RegionScopeLock(&StencilIdsCriticalSection);
+			StencilIds = MoveTemp(TempStencilIds);
+		}
+	}
+}
+
+void AColorCorrectRegion::HandleAffectedActorsPropertyChange()
+{
+	if (bActorListIsDirty)
+	{
+		bool bEventHandled = false;
+		bActorListIsDirty = false;
+		if (ActorListChangeType == EPropertyChangeType::ArrayAdd
+			|| ActorListChangeType == EPropertyChangeType::ValueSet)
+		{
+			bEventHandled = true;
+			if (ColorCorrectRegionsSubsystem)
+			{
+				ColorCorrectRegionsSubsystem->AssignStencilIdsToPerActorCC(this);
+			}
+		}
+
+		if (ActorListChangeType == EPropertyChangeType::ArrayClear
+			|| ActorListChangeType == EPropertyChangeType::ArrayRemove
+			|| ActorListChangeType == EPropertyChangeType::ValueSet)
+		{
+			bEventHandled = true;
+			if (ColorCorrectRegionsSubsystem)
+			{
+				ColorCorrectRegionsSubsystem->ClearStencilIdsToPerActorCC(this);
+			}
+		}
+	}
+}
+
 #if WITH_EDITOR
 void AColorCorrectRegion::PostEditChangeProperty(struct FPropertyChangedEvent& PropertyChangedEvent)
 {
@@ -155,6 +218,13 @@ void AColorCorrectRegion::PostEditChangeProperty(struct FPropertyChangedEvent& P
 		}
 	}
 
+	if (PropertyName == GET_MEMBER_NAME_CHECKED(AColorCorrectRegion, AffectedActors))
+	{
+		/** Since there might be Dialogs involved we need to dirty this CCR and handle the rest on Tick. */
+		ActorListChangeType = PropertyChangedEvent.ChangeType;
+		bActorListIsDirty = true;
+	}
+
 	// Reorder all CCRs after the Priority property has changed.
 	// Also, in context of Multi-User: PropertyChangedEvent can be a stub without the actual property data. 
 	// Therefore we need to refresh priority if PropertyChangedEvent.Property is nullptr. 
diff --git a/Engine/Plugins/Experimental/ColorCorrectRegions/Source/ColorCorrectRegions/Private/ColorCorrectRegionSubsystem.cpp b/Engine/Plugins/Experimental/ColorCorrectRegions/Source/ColorCorrectRegions/Private/ColorCorrectRegionSubsystem.cpp
--- a/Engine/Plugins/Experimental/ColorCorrectRegions/Source/ColorCorrectRegions/Private/ColorCorrectRegionSubsystem.cpp
+++ b/Engine/Plugins/Experimental/ColorCorrectRegions/Source/ColorCorrectRegions/Private/ColorCorrectRegionSubsystem.cpp
@@ -1,11 +1,15 @@
 // Copyright Epic Games, Inc. All Rights Reserved.
 #include "ColorCorrectRegionsSubsystem.h"
+#include "ColorCorrectRegionsModule.h"
+#include "ColorCorrectRegionsSceneViewExtension.h"
+#include "ColorCorrectRegionsStencilManager.h"
 #include "ColorCorrectWindow.h"
+
 #include "Engine/World.h"
 #include "Kismet/GameplayStatics.h"
 #include "EngineUtils.h"
 #include "SceneViewExtension.h"
-#include "ColorCorrectRegionsSceneViewExtension.h"
+
 
 #if WITH_EDITOR
 #include "Editor.h"
@@ -25,6 +29,7 @@ namespace
 #endif
 	}
 }
+
 void UColorCorrectRegionsSubsystem::Initialize(FSubsystemCollectionBase& Collection)
 {
 #if WITH_EDITOR
@@ -154,10 +159,36 @@ void UColorCorrectRegionsSubsystem::SortRegionsByDistance(const FVector& ViewLoc
 
 }
 
+void UColorCorrectRegionsSubsystem::AssignStencilIdsToPerActorCC(AColorCorrectRegion* Region, bool bIgnoreUserNotificaion, bool bSoftAssign)
+{
+	FColorCorrectRegionsStencilManager::AssignStencilIdsToAllActorsForCCR(GetWorld(), Region, bIgnoreUserNotificaion, bSoftAssign);
+}
+
+void UColorCorrectRegionsSubsystem::ClearStencilIdsToPerActorCC(AColorCorrectRegion* Region)
+{
+	FColorCorrectRegionsStencilManager::RemoveStencilNumberForSelectedRegion(GetWorld(), Region);
+}
+
+void UColorCorrectRegionsSubsystem::RefreshStenciIdAssignmentForAllCCR()
+{
+	for (TActorIterator<AColorCorrectRegion> It(GetWorld()); It; ++It)
+	{
+		AColorCorrectRegion* AsRegion = *It;
+		if (IsRegionValid(AsRegion, GetWorld()))
+		{
+			// Uncoment this if you want the invalid actors to be removed automatically.
+			// However after removal, undo/redo will not re-assign this actor back to CCR.
+			//FColorCorrectRegionsStencilManager::ClearInvalidActorsForSelectedRegion(AsRegion);
+			AsRegion->PerAffectedActorStencilData.Empty();
+			FColorCorrectRegionsStencilManager::AssignStencilIdsToAllActorsForCCR(GetWorld(), AsRegion, true, true);
+			FColorCorrectRegionsStencilManager::RemoveStencilNumberForSelectedRegion(GetWorld(), AsRegion);
+		}
+	}
+}
+
 void UColorCorrectRegionsSubsystem::RefreshRegions()
 {
 	FScopeLock RegionScopeLock(&RegionAccessCriticalSection);
-
 	RegionsPriorityBased.Reset();
 	RegionsDistanceBased.Reset();
 	for (TActorIterator<AColorCorrectRegion> It(GetWorld()); It; ++It)
@@ -175,6 +206,9 @@ void UColorCorrectRegionsSubsystem::RefreshRegions()
 			}
 		}
 	}
+
 	SortRegionsByPriority();
+
+	RefreshStenciIdAssignmentForAllCCR();
 }
 
diff --git a/Engine/Plugins/Experimental/ColorCorrectRegions/Source/ColorCorrectRegions/Private/ColorCorrectRegionsSceneViewExtension.cpp b/Engine/Plugins/Experimental/ColorCorrectRegions/Source/ColorCorrectRegions/Private/ColorCorrectRegionsSceneViewExtension.cpp
--- a/Engine/Plugins/Experimental/ColorCorrectRegions/Source/ColorCorrectRegions/Private/ColorCorrectRegionsSceneViewExtension.cpp
+++ b/Engine/Plugins/Experimental/ColorCorrectRegions/Source/ColorCorrectRegions/Private/ColorCorrectRegionsSceneViewExtension.cpp
@@ -244,7 +244,7 @@ namespace
 		, const FScreenPassTextureViewportParameters& SceneTextureViewportParams
 		, const FScreenPassTextureViewport& RegionViewport
 		, const FSceneTextureShaderParameters& SceneTextures
-		, const TArray<uint8>& StencilIds
+		, const TArray<uint32>& StencilIds
 		, FScreenPassRenderTarget& OutMergedStencilRenderTarget)
 	{
 		if (StencilIds.Num() == 0)
@@ -270,20 +270,11 @@ namespace
 			Parameters->PostProcessOutput = SceneTextureViewportParams;
 			Parameters->View = View.ViewUniformBuffer;
 
-			const FMatrix& RegionTransformMatrix = Region->GetActorTransform().ToMatrixWithScale();
-
 			FRHIBlendState* DefaultBlendState = FScreenPassPipelineState::FDefaultBlendState::GetRHI();
 				
 			FRHIResourceCreateInfo CreateInfo(TEXT("CCR_StencilIdBuffer"));
 
-			TArray<uint32> BufferData;
-			BufferData.AddUninitialized(StencilIds.Num());
-			for (int Index = 0; Index < BufferData.Num(); Index++)
-			{
-				BufferData[Index] = StencilIds[Index];
-			}
-
-			Parameters->StencilIds = GraphBuilder.CreateSRV(FRDGBufferSRVDesc(CreateStructuredBuffer(GraphBuilder, TEXT("CCR.StencilIdBuffer"), sizeof(uint32), StencilIds.Num(), &BufferData[0], sizeof(uint32) * StencilIds.Num())));
+			Parameters->StencilIds = GraphBuilder.CreateSRV(FRDGBufferSRVDesc(CreateStructuredBuffer(GraphBuilder, TEXT("CCR.StencilIdBuffer"), sizeof(uint32), StencilIds.Num(), &StencilIds[0], sizeof(uint32) * StencilIds.Num())));
 			Parameters->StencilIdCount = StencilIds.Num();
 
 			{
@@ -456,9 +447,10 @@ namespace
 			: static_cast<FColorCorrectRegionMaterialPS::ETemperatureType>(Region->TemperatureType);
 
 		FScreenPassRenderTarget MergedStencilRenderTarget;
-		if (Region->PerActorColorCorrection != EColorCorrectRegionStencilType::None)
+		if (Region->bEnablePerActorCC)
 		{
-			StencilMerger(GraphBuilder, Region, GlobalShaderMap, SceneColorRenderTarget, View, SceneTextureViewportParams, RegionViewport, SceneTextures, Region->StencilIds, MergedStencilRenderTarget);
+			TArray<uint32> StencilIds = Region->ConsumeStencilIds();
+			StencilMerger(GraphBuilder, Region, GlobalShaderMap, SceneColorRenderTarget, View, SceneTextureViewportParams, RegionViewport, SceneTextures, StencilIds, MergedStencilRenderTarget);
 		}
 
 		TShaderRef<FColorCorrectGenericPS> PixelShader;
diff --git a/Engine/Plugins/Experimental/ColorCorrectRegions/Source/ColorCorrectRegions/Private/ColorCorrectRegionsStencilManager.cpp b/Engine/Plugins/Experimental/ColorCorrectRegions/Source/ColorCorrectRegions/Private/ColorCorrectRegionsStencilManager.cpp
new file mode 100644
--- /dev/null
+++ b/Engine/Plugins/Experimental/ColorCorrectRegions/Source/ColorCorrectRegions/Private/ColorCorrectRegionsStencilManager.cpp
@@ -0,0 +1,315 @@
+// Copyright Epic Games, Inc. All Rights Reserved.
+
+#include "ColorCorrectRegionsStencilManager.h"
+#include "ColorCorrectRegionsModule.h"
+#include "Engine/GameEngine.h"
+#include "Misc/MessageDialog.h"
+#include "Components/PrimitiveComponent.h"
+#include "EngineUtils.h"
+#include "Engine/World.h"
+
+
+
+#define STENCILID_MIN 127
+#define STENCILID_MAX 255
+
+namespace
+{
+	/*
+	* returns true if Actor is being used by another CCR and is currently sharing the same stencil number with other actors assigned to that CCR.
+	* This means that we need to assign a new stencil number to this actor across all CCRs that use it.
+	*
+	* Also returns OutStencilDataToUse which is shared information about this actor between different CCRs. If OutStencilDataToUse is invalid
+	* that means that there are no other CCRs that stencil this actor.
+	*/
+	static bool FindIfActorIsUsedByCCR(TSoftObjectPtr<AActor> ActorToSearchFor, const TArray<AActor*>& CCRsToSearchIn, TSharedPtr<FStencilData>& OutStencilDataToUse)
+	{
+		bool bAssignNewStencilId = false;
+		OutStencilDataToUse = nullptr;
+
+		for (const AActor* ActorCCR : CCRsToSearchIn)
+		{
+			const AColorCorrectRegion* CCR = Cast<AColorCorrectRegion>(ActorCCR);
+			if (CCR->PerAffectedActorStencilData.Contains(ActorToSearchFor))
+			{
+				OutStencilDataToUse = CCR->PerAffectedActorStencilData[ActorToSearchFor];
+				if (CCR->PerAffectedActorStencilData.Num() > 1)
+				{
+					TSharedPtr<FStencilData> TempStencilData = CCR->PerAffectedActorStencilData[ActorToSearchFor];
+					for (const TPair<TSoftObjectPtr<AActor>, TSharedPtr<FStencilData>>& StencilDataPair : CCR->PerAffectedActorStencilData)
+					{
+						if (StencilDataPair.Key != ActorToSearchFor && StencilDataPair.Value->AssignedStencil == TempStencilData->AssignedStencil)
+						{
+							bAssignNewStencilId = true;
+							return bAssignNewStencilId;
+						}
+					}
+				}
+			}
+		}
+		return bAssignNewStencilId;
+	}
+
+	/*
+	* @param bIgnoreUserNotification	In cases when another CCR is already using this actor, we don't want to notify the user. Another example is when we reload the scene,
+	* when all we need to do is reassign stencil ids.
+	* @param bSoftAssign			In some cases we don't want to assign stencil ids to actors (for example when we reload the scene).
+	*/
+	static void AssignStencilIdToActor(AColorCorrectRegion* InRegion, bool bIgnoreUserNotification, bool bSoftAssign, TSoftObjectPtr<AActor> ActorToAssignStencilTo, TSharedPtr<FStencilData> StencilData)
+	{
+		TArray<UPrimitiveComponent*> PrimitiveComponents;
+		ActorToAssignStencilTo->GetComponents<UPrimitiveComponent>(PrimitiveComponents);
+		for (UPrimitiveComponent* PrimitiveComponent : PrimitiveComponents)
+		{
+			// If the stencil id is already assigned and not by another CCR then we need to notify the user.
+			if ((PrimitiveComponent->bRenderCustomDepth && PrimitiveComponent->CustomDepthStencilValue < STENCILID_MIN))
+			{
+				EAppReturnType::Type Answer = EAppReturnType::No;
+				if (!bIgnoreUserNotification)
+				{
+					Answer = FMessageDialog::Open(EAppMsgType::YesNo,
+						FText::FromString(FString::Printf(TEXT("The actor %s already has Custom Depth Stencil Id assigned. \nThis operation will overwrite the current value and could potentially affect the visuals.\nAre you sure?"), *ActorToAssignStencilTo->GetName())));
+
+				}
+				if (Answer == EAppReturnType::No)
+				{
+					StencilData->AssignedStencil = PrimitiveComponent->CustomDepthStencilValue;
+					InRegion->PerAffectedActorStencilData.Add(ActorToAssignStencilTo, StencilData);
+					return;
+				}
+			}
+			if (!bSoftAssign)
+			{
+				PrimitiveComponent->SetRenderCustomDepth(true);
+				PrimitiveComponent->CustomDepthStencilValue = static_cast<int32>(StencilData->AssignedStencil);
+				PrimitiveComponent->MarkRenderStateDirty();
+			}
+			else
+			{
+				StencilData->AssignedStencil = PrimitiveComponent->CustomDepthStencilValue;
+			}
+		}
+		InRegion->PerAffectedActorStencilData.Add(ActorToAssignStencilTo, StencilData);
+	}
+
+	/*
+	* Handles the cleanup for the provided actor.
+	*/
+	static void ClearStencilIdFromActor(TSoftObjectPtr<AActor> ActorToAssignStencilTo)
+	{
+		if (ActorToAssignStencilTo.IsValid())
+		{
+			TArray<UPrimitiveComponent*> PrimitiveComponents;
+			ActorToAssignStencilTo->GetComponents<UPrimitiveComponent>(PrimitiveComponents);
+
+			for (UPrimitiveComponent* PrimitiveComponent : PrimitiveComponents)
+			{
+				// If the stencil id is within our range then it is managed by us, so we should reset this to default.
+				if (PrimitiveComponent->bRenderCustomDepth && PrimitiveComponent->CustomDepthStencilValue >= STENCILID_MIN)
+				{
+					PrimitiveComponent->SetRenderCustomDepth(false);
+					PrimitiveComponent->CustomDepthStencilValue = 0;
+					PrimitiveComponent->MarkRenderStateDirty();
+				}
+			}
+		}
+	}
+
+	/*
+	* Simple checking function that returns true if the provided actor is used by another CCR.
+	*/
+	static bool IsActorAssignedToOtherCCR(AColorCorrectRegion* CurrentCCR, TSoftObjectPtr<AActor> ActorToSearchFor, const TArray<AActor*>& CCRsToSearchIn)
+	{
+		bool bAssignNewStencilId = false;
+
+		for (const AActor* ActorCCR : CCRsToSearchIn)
+		{
+			const AColorCorrectRegion* CCR = Cast<AColorCorrectRegion>(ActorCCR);
+			if (CurrentCCR == CCR)
+			{
+				continue;
+			}
+			if (CCR->PerAffectedActorStencilData.Contains(ActorToSearchFor))
+			{
+				return true;
+			}
+		}
+		return false;
+	}
+}
+
+
+void FColorCorrectRegionsStencilManager::AssignStencilNumberToActorForSelectedRegion(UWorld* CurrentWorld, AColorCorrectRegion* Region, TSoftObjectPtr<AActor> ActorToAssignStencilTo, bool bIgnoreUserNotifications, bool bSoftAssign)
+{
+	verify(ActorToAssignStencilTo.IsValid());
+	TArray<FStencilData> StencilDatas;
+
+	ULevel* CurrentLevel = CurrentWorld->GetCurrentLevel();
+	if (Region->PerAffectedActorStencilData.Contains(ActorToAssignStencilTo) && Region->PerAffectedActorStencilData[ActorToAssignStencilTo]->AssignedStencil != 0)
+	{
+		// This actor is already assigned to the current CCR. No need to do anything.
+		return;
+	}
+
+	// This array is an optimization that avoids additional costs of going over all actors in the scene and searching for CCRs.
+	TArray<AActor*> AllCCRsInCurrentLevel;
+
+	TSharedPtr<FStencilData> StencilDataToUse;
+	TSortedMap<uint32, TSharedPtr<FStencilData>> AllUsedIds;
+
+	// Gather all unique stencil ids that are currently in use.
+	for (TActorIterator<AColorCorrectRegion> It(CurrentWorld); It; ++It)
+	{
+		AColorCorrectRegion* CCR = *It;
+		AllCCRsInCurrentLevel.Add(CCR);
+		for (const TPair<TSoftObjectPtr<AActor>, TSharedPtr<FStencilData>>& StencilDataPair : CCR->PerAffectedActorStencilData)
+		{
+			AllUsedIds.Add(StencilDataPair.Value->AssignedStencil, StencilDataPair.Value);
+		}
+	}
+
+	// Find if stencil number is already assigned to the current actor by another CCR.
+	bool bAssignNewStencilId = FindIfActorIsUsedByCCR(ActorToAssignStencilTo, AllCCRsInCurrentLevel, StencilDataToUse);
+
+	bool bUsedByAnotherCCR = StencilDataToUse.IsValid();
+
+	/*
+	* If the actor is not used by any other CCR, it means that we need to see if we can group up stencil id
+	* for that actor with another, that already belongs to the current CCR,
+	*
+	* Find Actors that are assigned to the current CCR and make sure these are unused by any other CCR before
+	* using that stencil ID
+	*/
+	if (!StencilDataToUse.IsValid())
+	{
+		StencilDataToUse = MakeShared<FStencilData>();
+		bAssignNewStencilId = true;
+
+		if (Region->PerAffectedActorStencilData.Num() > 0)
+		{
+			for (TPair<TSoftObjectPtr<AActor>, TSharedPtr<FStencilData>> ActorDataPair : Region->PerAffectedActorStencilData)
+			{
+				if (!IsActorAssignedToOtherCCR(Region, ActorDataPair.Key, AllCCRsInCurrentLevel))
+				{
+					if (ActorDataPair.Value->AssignedStencil >= STENCILID_MIN
+						&& ActorDataPair.Value->AssignedStencil <= STENCILID_MAX)
+					{
+						// we reuse stencil Id of another actor which is assigned to the same CCR. 
+						bAssignNewStencilId = false;
+						StencilDataToUse->AssignedStencil = ActorDataPair.Value->AssignedStencil;
+						break;
+					}
+				}
+			}
+		}
+	}
+
+	/*
+	* Stencil Id needs doesn't need to be assigned, for either of the following reasons:
+	* - We reuse stencil Id from another actor that is assigned to the same CCR (Group up ids)
+	* - This actor is used by another CCR and that CCR doesn't have any other actors, so we can reuse the same id.
+	*/
+	if (!bAssignNewStencilId)
+	{
+		AssignStencilIdToActor(Region, bUsedByAnotherCCR || bIgnoreUserNotifications, bSoftAssign, ActorToAssignStencilTo, StencilDataToUse);
+		return;
+	}
+
+	if (bAssignNewStencilId)
+	{
+		uint8 NextAvailableStencilId = STENCILID_MIN;
+		while (NextAvailableStencilId <= STENCILID_MAX)
+		{
+			if (!AllUsedIds.Contains(NextAvailableStencilId))
+			{
+				break;
+			}
+			NextAvailableStencilId++;
+		}
+		if (NextAvailableStencilId > STENCILID_MAX)
+		{
+			UE_LOG(ColorCorrectRegions, Error, TEXT("Run out of Custom Depth Stencil ID values to be used. The maximum stencil ID is assigned. This may yeild unexpected results. "));
+			NextAvailableStencilId = STENCILID_MAX;
+		}
+		StencilDataToUse->AssignedStencil = NextAvailableStencilId;
+		AssignStencilIdToActor(Region, bUsedByAnotherCCR || bIgnoreUserNotifications, bSoftAssign, ActorToAssignStencilTo, StencilDataToUse);
+	}
+}
+
+void FColorCorrectRegionsStencilManager::RemoveStencilNumberForSelectedRegion(UWorld* CurrentWorld, AColorCorrectRegion* Region)
+{
+	TSet<TSoftObjectPtr<AActor>> ActorsToCleanup;
+
+	// Accumulate all CCRs that were removed.
+	for (TPair<TSoftObjectPtr<AActor>, TSharedPtr<FStencilData>> ActorDataPair : Region->PerAffectedActorStencilData)
+	{
+		if (!Region->AffectedActors.Contains(ActorDataPair.Key))
+		{
+			ActorsToCleanup.Add(ActorDataPair.Key);
+		}
+	}
+
+	ULevel* CurrentLevel = CurrentWorld->GetCurrentLevel();
+	TArray<AActor*> AllCCRsInCurrentLevel = CurrentLevel->Actors.FilterByPredicate([](const AActor* Actor) { return Cast<AColorCorrectRegion>(Actor) != nullptr; });
+
+	// For each actor that was removed, remove Ids if it is unused by any other CCR, or otherwise assign a new id.
+	for (TSoftObjectPtr<AActor> ActorToRemove : ActorsToCleanup)
+	{
+		Region->PerAffectedActorStencilData.Remove(ActorToRemove);
+
+		bool bUsedByAnotherActor = false;
+		for (AActor* ActorCCR : AllCCRsInCurrentLevel)
+		{
+			AColorCorrectRegion* OtherCCR = Cast<AColorCorrectRegion>(ActorCCR);
+
+			if (OtherCCR->PerAffectedActorStencilData.Contains(ActorToRemove))
+			{
+				OtherCCR->PerAffectedActorStencilData.Remove(ActorToRemove);
+				AssignStencilNumberToActorForSelectedRegion(CurrentWorld, OtherCCR, ActorToRemove, true, false);
+				bUsedByAnotherActor = true;
+				break;
+			}
+		}
+
+		// This actor is not used by any aother actor. Remove it.
+		if (!bUsedByAnotherActor)
+		{
+			ClearStencilIdFromActor(ActorToRemove);
+		}
+	}
+}
+
+void FColorCorrectRegionsStencilManager::ClearInvalidActorsForSelectedRegion(AColorCorrectRegion* Region)
+{
+	TArray<TSoftObjectPtr<AActor>> ActorsToCleanup;
+	for (TSoftObjectPtr<AActor> Actor : Region->AffectedActors)
+	{
+		if (!Actor.IsValid())
+		{
+			ActorsToCleanup.Add(Actor);
+		}
+	}
+	for (TSoftObjectPtr<AActor> Actor : ActorsToCleanup)
+	{
+		Region->AffectedActors.Remove(Actor);
+	}
+}
+
+void FColorCorrectRegionsStencilManager::AssignStencilIdsToAllActorsForCCR(UWorld* CurrentWorld, AColorCorrectRegion* Region, bool bIgnoreUserNotifications, bool bSoftAssign)
+{
+	if (Region->AffectedActors.Num() > 0)
+	{
+		for (TSoftObjectPtr<AActor> Actor : Region->AffectedActors)
+		{
+			if (Actor.IsValid())
+			{
+				FColorCorrectRegionsStencilManager::AssignStencilNumberToActorForSelectedRegion(CurrentWorld, Region, Actor, bIgnoreUserNotifications, bSoftAssign);
+			}
+		}
+	}
+}
+
+
+#undef STENCILID_MIN
+#undef STENCILID_MAX
+
diff --git a/Engine/Plugins/Experimental/ColorCorrectRegions/Source/ColorCorrectRegions/Public/ColorCorrectRegion.h b/Engine/Plugins/Experimental/ColorCorrectRegions/Source/ColorCorrectRegions/Public/ColorCorrectRegion.h
--- a/Engine/Plugins/Experimental/ColorCorrectRegions/Source/ColorCorrectRegions/Public/ColorCorrectRegion.h
+++ b/Engine/Plugins/Experimental/ColorCorrectRegions/Source/ColorCorrectRegions/Public/ColorCorrectRegion.h
@@ -7,6 +7,7 @@
 #include "GameFramework/Actor.h"
 #include "Engine/Classes/Components/MeshComponent.h"
 #include "Engine/Scene.h"
+#include "Templates/SharedPointer.h"
 #include "ColorCorrectRegion.generated.h"
 
 
@@ -35,12 +36,18 @@ enum class EColorCorrectRegionTemperatureType : uint8
 UENUM(BlueprintType)
 enum class EColorCorrectRegionStencilType : uint8
 {
-	None					UMETA(DisplayName = "Default (Affect all actors)"),
 	ExcludeStencil			UMETA(DisplayName = "Exclude Selected Actors"),
 	IncludeStencil			UMETA(DisplayName = "Affect Only Selected Actors"),
 	MAX
 };
 
+class FStencilData
+{
+public:
+	uint32 AssignedStencil = 0;
+};
+
+
 /**
  * An instance of Color Correction Region. Used to aggregate all active regions.
  * This actor is aggregated by ColorCorrectRegionsSubsystem which handles:
@@ -107,13 +114,24 @@ public:
 	UPROPERTY(BlueprintReadWrite, EditAnywhere, Category = "Color Correction")
 	bool Enabled;
 
+	/** Enables or disabled per actor color correction. */
+	UPROPERTY(BlueprintReadWrite, EditAnywhere, Category = "Per Actor CC")
+	bool bEnablePerActorCC;
+
 	/** Controls in which way the below targets will be affected by color correction. */
-	UPROPERTY(BlueprintReadWrite, EditAnywhere, Category = "Color Correction")
+	UPROPERTY(BlueprintReadWrite, EditAnywhere, Category = "Per Actor CC", meta = (editcondition = "bEnablePerActorCC", DisplayName = "Per Actor CC Mode"))
 	EColorCorrectRegionStencilType PerActorColorCorrection;
 
-	/** A list of actors that need to be excluded or included from color correction (Depends on the option above). */
-	UPROPERTY(BlueprintReadWrite, EditAnywhere, Category = "Color Correction")
-	TArray<uint8> StencilIds;
+	/** List of actors that get affected or ignored by Per actor CC. Effect depends on the above option. */
+	UPROPERTY(BlueprintReadWrite, EditAnywhere, Category = "Per Actor CC", meta = (editcondition = "bEnablePerActorCC", DisplayName = "Actor Selection"))
+	TSet<TSoftObjectPtr<AActor>> AffectedActors;
+
+	/** 
+	* This is a temporary property that is used to keep track of managed actors for per actor CC.
+	* PerAffectedActorStencilData is managed by FColorCorrectRegionsStencilManager and CCR Subsystem
+	* Its lifetime is aligned with selected level and gets reset often.
+	*/
+	TMap<TSoftObjectPtr<AActor>, TSharedPtr<FStencilData>> PerAffectedActorStencilData;
 
 #if WITH_EDITORONLY_DATA
 
@@ -163,6 +181,22 @@ public:
 		return FirstPrimitiveId;
 	};
 
+	TArray<uint32> ConsumeStencilIds() 
+	{ 
+		FScopeLock RegionScopeLock(&StencilIdsCriticalSection);
+		return MoveTemp(StencilIds); 
+	};
+private:
+	/**
+	* AffectedActors property change could potentially invoke a Dialog Window, which should be displayed on Game Thread.
+	*/
+	void HandleAffectedActorsPropertyChange();
+
+	/**
+	* Copy stencil Ids from Primitive components to the StencilId array to be consumed by Scene view extension.
+	*/
+	void TransferStencilIds();
+
 private:
 	UColorCorrectRegionsSubsystem* ColorCorrectRegionsSubsystem;
 
@@ -171,4 +205,18 @@ private:
 
 	FTransform PreviousFrameTransform;
 	FPrimitiveComponentId FirstPrimitiveId;
+
+	bool bActorListIsDirty;
+	uint8 ActorListChangeType;
+
+	/**
+	* A list of actors that need to be excluded or included from color correction to be used on render thread.
+	* This list is populated during engine tick and populated with Stencil Ids from Primitive Components of AffectedActors.
+	* The main reason for this is acessing Primitive Components is not thread safe and is better to be accessed on Game thread.
+	* 
+	* Stencil Ids are moved on render thread to avoid locking either threads for too long.
+	*/
+	TArray<uint32> StencilIds;
+	FCriticalSection StencilIdsCriticalSection;
+
 };
diff --git a/Engine/Plugins/Experimental/ColorCorrectRegions/Source/ColorCorrectRegions/Public/ColorCorrectRegionsStencilManager.h b/Engine/Plugins/Experimental/ColorCorrectRegions/Source/ColorCorrectRegions/Public/ColorCorrectRegionsStencilManager.h
new file mode 100644
--- /dev/null
+++ b/Engine/Plugins/Experimental/ColorCorrectRegions/Source/ColorCorrectRegions/Public/ColorCorrectRegionsStencilManager.h
@@ -0,0 +1,18 @@
+// Copyright Epic Games, Inc. All Rights Reserved.
+
+#pragma once
+
+#include "CoreMinimal.h"
+#include "ColorCorrectRegion.h"
+#include "GameFramework/Actor.h"
+
+
+class FColorCorrectRegionsStencilManager
+{
+public:
+		static void AssignStencilNumberToActorForSelectedRegion(UWorld* CurrentWorld, AColorCorrectRegion* Region, TSoftObjectPtr<AActor> ActorToAssignStencilTo, bool bIgnoreUserNotifications, bool bSoftAssign);
+		static void RemoveStencilNumberForSelectedRegion(UWorld* CurrentWorld, AColorCorrectRegion* Region);
+		static void ClearInvalidActorsForSelectedRegion(AColorCorrectRegion* Region);
+		static void AssignStencilIdsToAllActorsForCCR(UWorld* CurrentWorld, AColorCorrectRegion* Region, bool bIgnoreUserNotifications, bool bSoftAssign);
+
+};
\ No newline at end of file
diff --git a/Engine/Plugins/Experimental/ColorCorrectRegions/Source/ColorCorrectRegions/Public/ColorCorrectRegionsSubsystem.h b/Engine/Plugins/Experimental/ColorCorrectRegions/Source/ColorCorrectRegions/Public/ColorCorrectRegionsSubsystem.h
--- a/Engine/Plugins/Experimental/ColorCorrectRegions/Source/ColorCorrectRegions/Public/ColorCorrectRegionsSubsystem.h
+++ b/Engine/Plugins/Experimental/ColorCorrectRegions/Source/ColorCorrectRegions/Public/ColorCorrectRegionsSubsystem.h
@@ -83,6 +83,16 @@ public:
 	/** Sorts regions based on distance from the camera. */
 	void SortRegionsByDistance(const FVector& ViewLocation);
 
+	/** Handles Stencil Ids for the selected CCR and corresponding actor. */
+	void AssignStencilIdsToPerActorCC(AColorCorrectRegion* Region, bool bIgnoreUserNotificaion = false, bool bSoftAssign = false);
+	
+	/** Handles removal of Stencil Ids for the selected CCR. */
+	void ClearStencilIdsToPerActorCC(AColorCorrectRegion* Region);
+
+	/** Resets all stencils and re-assigns for each CCR in the scene. */
+	UFUNCTION(BlueprintCallable, meta = (Category = "Color Correct Regions"))
+	void RefreshStenciIdAssignmentForAllCCR();
+
 private:
 
 	/** Repopulates array of region actors. */
