Color Correct Regions:
- fixed CCR flashing due to race condition.
- Added appropriate reaction to when stencil Ids are changed by the user manually.
- Added debug timings for Insights.

diff --git a/Engine/Plugins/Experimental/ColorCorrectRegions/Source/ColorCorrectRegions/Private/ColorCorrectRegion.cpp b/Engine/Plugins/Experimental/ColorCorrectRegions/Source/ColorCorrectRegions/Private/ColorCorrectRegion.cpp
--- a/Engine/Plugins/Experimental/ColorCorrectRegions/Source/ColorCorrectRegions/Private/ColorCorrectRegion.cpp
+++ b/Engine/Plugins/Experimental/ColorCorrectRegions/Source/ColorCorrectRegions/Private/ColorCorrectRegion.cpp
@@ -122,9 +122,32 @@ void AColorCorrectRegion::TickActor(float DeltaTime, ELevelTick TickType, FActor
 {
 	Super::Tick(DeltaTime);
 
+	TRACE_CPUPROFILER_EVENT_SCOPE_TEXT(*FString::Printf(TEXT("CCR.TickActor %s"), *GetName()));
+
 	HandleAffectedActorsPropertyChange();
 	TransferStencilIds();
 
+	// Check to make sure that no ids have been changed externally.
+	{
+		TimeWaited += DeltaTime;
+		const float WaitTimeInSecs = 1.0;
+
+		if (!ColorCorrectRegionsSubsystem)
+		{
+			if (const UWorld* World = GetWorld())
+			{
+				ColorCorrectRegionsSubsystem = World->GetSubsystem<UColorCorrectRegionsSubsystem>();
+			}
+		}
+		
+		if (ColorCorrectRegionsSubsystem && TimeWaited >= WaitTimeInSecs)
+		{
+			ColorCorrectRegionsSubsystem->CheckAssignedActorsValidity(this);
+			TimeWaited = 0;
+		}
+
+	}
+
 	FTransform CurrentFrameTransform = GetTransform();
 	if (!PreviousFrameTransform.Equals(CurrentFrameTransform))
 	{
diff --git a/Engine/Plugins/Experimental/ColorCorrectRegions/Source/ColorCorrectRegions/Private/ColorCorrectRegionSubsystem.cpp b/Engine/Plugins/Experimental/ColorCorrectRegions/Source/ColorCorrectRegions/Private/ColorCorrectRegionSubsystem.cpp
--- a/Engine/Plugins/Experimental/ColorCorrectRegions/Source/ColorCorrectRegions/Private/ColorCorrectRegionSubsystem.cpp
+++ b/Engine/Plugins/Experimental/ColorCorrectRegions/Source/ColorCorrectRegions/Private/ColorCorrectRegionSubsystem.cpp
@@ -169,6 +169,11 @@ void UColorCorrectRegionsSubsystem::ClearStencilIdsToPerActorCC(AColorCorrectReg
 	FColorCorrectRegionsStencilManager::RemoveStencilNumberForSelectedRegion(GetWorld(), Region);
 }
 
+void UColorCorrectRegionsSubsystem::CheckAssignedActorsValidity(AColorCorrectRegion* Region)
+{
+	FColorCorrectRegionsStencilManager::CheckAssignedActorsValidity(Region);
+}
+
 void UColorCorrectRegionsSubsystem::RefreshStenciIdAssignmentForAllCCR()
 {
 	for (TActorIterator<AColorCorrectRegion> It(GetWorld()); It; ++It)
diff --git a/Engine/Plugins/Experimental/ColorCorrectRegions/Source/ColorCorrectRegions/Private/ColorCorrectRegionsSceneViewExtension.cpp b/Engine/Plugins/Experimental/ColorCorrectRegions/Source/ColorCorrectRegions/Private/ColorCorrectRegionsSceneViewExtension.cpp
--- a/Engine/Plugins/Experimental/ColorCorrectRegions/Source/ColorCorrectRegions/Private/ColorCorrectRegionsSceneViewExtension.cpp
+++ b/Engine/Plugins/Experimental/ColorCorrectRegions/Source/ColorCorrectRegions/Private/ColorCorrectRegionsSceneViewExtension.cpp
@@ -449,7 +449,7 @@ namespace
 		FScreenPassRenderTarget MergedStencilRenderTarget;
 		if (Region->bEnablePerActorCC)
 		{
-			TArray<uint32> StencilIds = Region->ConsumeStencilIds();
+			TArray<uint32> StencilIds = Region->GetCopyStencilIds();
 			StencilMerger(GraphBuilder, Region, GlobalShaderMap, SceneColorRenderTarget, View, SceneTextureViewportParams, RegionViewport, SceneTextures, StencilIds, MergedStencilRenderTarget);
 		}
 
diff --git a/Engine/Plugins/Experimental/ColorCorrectRegions/Source/ColorCorrectRegions/Private/ColorCorrectRegionsStencilManager.cpp b/Engine/Plugins/Experimental/ColorCorrectRegions/Source/ColorCorrectRegions/Private/ColorCorrectRegionsStencilManager.cpp
--- a/Engine/Plugins/Experimental/ColorCorrectRegions/Source/ColorCorrectRegions/Private/ColorCorrectRegionsStencilManager.cpp
+++ b/Engine/Plugins/Experimental/ColorCorrectRegions/Source/ColorCorrectRegions/Private/ColorCorrectRegionsStencilManager.cpp
@@ -10,8 +10,17 @@
 
 
 
-#define STENCILID_MIN 127
-#define STENCILID_MAX 255
+static TAutoConsoleVariable<int32> CVarStencilIdRangeMin(
+	TEXT("r.CCR.StencilIdRangeMin"),
+	127,
+	TEXT("The lower boundary of stencil Ids that are reserved by CCR.\n\
+		By default this value is set to 127. Must be within the range of [0, 255]."));
+
+static TAutoConsoleVariable<int32> CVarStencilIdRangeMax(
+	TEXT("r.CCR.StencilIdRangeMax"),
+	255,
+	TEXT("The upper boundary of stencil Ids that are reserved by CCR. \n\
+		By default this value is set to 255. Must be within the range of [0, 255].\n"));
 
 namespace
 {
@@ -59,16 +68,21 @@ namespace
 	{
 		TArray<UPrimitiveComponent*> PrimitiveComponents;
 		ActorToAssignStencilTo->GetComponents<UPrimitiveComponent>(PrimitiveComponents);
+
+		uint8 StencilIdMin = FMath::Min(CVarStencilIdRangeMin.GetValueOnAnyThread(), CVarStencilIdRangeMax.GetValueOnAnyThread());
+		uint8 StencilIdMax = FMath::Max(CVarStencilIdRangeMin.GetValueOnAnyThread(), CVarStencilIdRangeMax.GetValueOnAnyThread());
 		for (UPrimitiveComponent* PrimitiveComponent : PrimitiveComponents)
 		{
 			// If the stencil id is already assigned and not by another CCR then we need to notify the user.
-			if ((PrimitiveComponent->bRenderCustomDepth && PrimitiveComponent->CustomDepthStencilValue < STENCILID_MIN))
+			if ((PrimitiveComponent->bRenderCustomDepth && PrimitiveComponent->CustomDepthStencilValue < StencilIdMin))
 			{
 				EAppReturnType::Type Answer = EAppReturnType::No;
 				if (!bIgnoreUserNotification)
 				{
-					Answer = FMessageDialog::Open(EAppMsgType::YesNo,
-						FText::FromString(FString::Printf(TEXT("The actor %s already has Custom Depth Stencil Id assigned. \nThis operation will overwrite the current value and could potentially affect the visuals.\nAre you sure?"), *ActorToAssignStencilTo->GetName())));
+					auto Text = FString::Printf(TEXT("The actor \'%s\' already has Custom Depth Stencil Id assigned that is out of range.\
+													\nWould you like for Color Correct Regions to manage Stencil Id?\
+													\nThe value will be reassigned to within managed range [%d %d]."), *ActorToAssignStencilTo->GetName(), StencilIdMin, StencilIdMax);
+					Answer = FMessageDialog::Open(EAppMsgType::YesNo, FText::FromString(Text));
 
 				}
 				if (Answer == EAppReturnType::No)
@@ -97,6 +111,7 @@ namespace
 	*/
 	static void ClearStencilIdFromActor(TSoftObjectPtr<AActor> ActorToAssignStencilTo)
 	{
+		uint8 StencilIdMin = FMath::Min(CVarStencilIdRangeMin.GetValueOnAnyThread(), CVarStencilIdRangeMax.GetValueOnAnyThread());
 		if (ActorToAssignStencilTo.IsValid())
 		{
 			TArray<UPrimitiveComponent*> PrimitiveComponents;
@@ -105,7 +120,7 @@ namespace
 			for (UPrimitiveComponent* PrimitiveComponent : PrimitiveComponents)
 			{
 				// If the stencil id is within our range then it is managed by us, so we should reset this to default.
-				if (PrimitiveComponent->bRenderCustomDepth && PrimitiveComponent->CustomDepthStencilValue >= STENCILID_MIN)
+				if (PrimitiveComponent->bRenderCustomDepth && PrimitiveComponent->CustomDepthStencilValue >= StencilIdMin)
 				{
 					PrimitiveComponent->SetRenderCustomDepth(false);
 					PrimitiveComponent->CustomDepthStencilValue = 0;
@@ -141,7 +156,9 @@ namespace
 
 void FColorCorrectRegionsStencilManager::AssignStencilNumberToActorForSelectedRegion(UWorld* CurrentWorld, AColorCorrectRegion* Region, TSoftObjectPtr<AActor> ActorToAssignStencilTo, bool bIgnoreUserNotifications, bool bSoftAssign)
 {
+	TRACE_CPUPROFILER_EVENT_SCOPE_TEXT(*FString::Printf(TEXT("CCR.StencilManager.AssignStencilNumberToActorForSelectedRegion %s"), *Region->GetName()));
 	verify(ActorToAssignStencilTo.IsValid());
+	
 	TArray<FStencilData> StencilDatas;
 
 	ULevel* CurrentLevel = CurrentWorld->GetCurrentLevel();
@@ -173,6 +190,9 @@ void FColorCorrectRegionsStencilManager::AssignStencilNumberToActorForSelectedRe
 
 	bool bUsedByAnotherCCR = StencilDataToUse.IsValid();
 
+	uint8 StencilIdMin = FMath::Min(CVarStencilIdRangeMin.GetValueOnAnyThread(), CVarStencilIdRangeMax.GetValueOnAnyThread());
+	uint8 StencilIdMax = FMath::Max(CVarStencilIdRangeMin.GetValueOnAnyThread(), CVarStencilIdRangeMax.GetValueOnAnyThread());
+
 	/*
 	* If the actor is not used by any other CCR, it means that we need to see if we can group up stencil id
 	* for that actor with another, that already belongs to the current CCR,
@@ -182,7 +202,7 @@ void FColorCorrectRegionsStencilManager::AssignStencilNumberToActorForSelectedRe
 	*/
 	if (!StencilDataToUse.IsValid())
 	{
-		StencilDataToUse = MakeShared<FStencilData>();
+		StencilDataToUse = MakeShared<FStencilData>();		
 		bAssignNewStencilId = true;
 
 		if (Region->PerAffectedActorStencilData.Num() > 0)
@@ -191,8 +211,8 @@ void FColorCorrectRegionsStencilManager::AssignStencilNumberToActorForSelectedRe
 			{
 				if (!IsActorAssignedToOtherCCR(Region, ActorDataPair.Key, AllCCRsInCurrentLevel))
 				{
-					if (ActorDataPair.Value->AssignedStencil >= STENCILID_MIN
-						&& ActorDataPair.Value->AssignedStencil <= STENCILID_MAX)
+					if (ActorDataPair.Value->AssignedStencil >= StencilIdMin
+						&& ActorDataPair.Value->AssignedStencil <= StencilIdMax)
 					{
 						// we reuse stencil Id of another actor which is assigned to the same CCR. 
 						bAssignNewStencilId = false;
@@ -217,8 +237,9 @@ void FColorCorrectRegionsStencilManager::AssignStencilNumberToActorForSelectedRe
 
 	if (bAssignNewStencilId)
 	{
-		uint8 NextAvailableStencilId = STENCILID_MIN;
-		while (NextAvailableStencilId <= STENCILID_MAX)
+		uint8 NextAvailableStencilId = StencilIdMin;
+
+		while (NextAvailableStencilId <= StencilIdMax)
 		{
 			if (!AllUsedIds.Contains(NextAvailableStencilId))
 			{
@@ -226,10 +247,10 @@ void FColorCorrectRegionsStencilManager::AssignStencilNumberToActorForSelectedRe
 			}
 			NextAvailableStencilId++;
 		}
-		if (NextAvailableStencilId > STENCILID_MAX)
+		if (NextAvailableStencilId > StencilIdMax)
 		{
 			UE_LOG(ColorCorrectRegions, Error, TEXT("Run out of Custom Depth Stencil ID values to be used. The maximum stencil ID is assigned. This may yeild unexpected results. "));
-			NextAvailableStencilId = STENCILID_MAX;
+			NextAvailableStencilId = StencilIdMax;
 		}
 		StencilDataToUse->AssignedStencil = NextAvailableStencilId;
 		AssignStencilIdToActor(Region, bUsedByAnotherCCR || bIgnoreUserNotifications, bSoftAssign, ActorToAssignStencilTo, StencilDataToUse);
@@ -238,6 +259,7 @@ void FColorCorrectRegionsStencilManager::AssignStencilNumberToActorForSelectedRe
 
 void FColorCorrectRegionsStencilManager::RemoveStencilNumberForSelectedRegion(UWorld* CurrentWorld, AColorCorrectRegion* Region)
 {
+	TRACE_CPUPROFILER_EVENT_SCOPE_TEXT(*FString::Printf(TEXT("CCR.StencilManager.RemoveStencilNumberForSelectedRegion %s"), *Region->GetName()));
 	TSet<TSoftObjectPtr<AActor>> ActorsToCleanup;
 
 	// Accumulate all CCRs that were removed.
@@ -309,7 +331,47 @@ void FColorCorrectRegionsStencilManager::AssignStencilIdsToAllActorsForCCR(UWorl
 	}
 }
 
+void FColorCorrectRegionsStencilManager::CheckAssignedActorsValidity(AColorCorrectRegion* Region)
+{
+	TRACE_CPUPROFILER_EVENT_SCOPE_TEXT(*FString::Printf(TEXT("CCR.StencilManager.CheckAssignedActorsValidity %s"), *Region->GetName()));
 
-#undef STENCILID_MIN
-#undef STENCILID_MAX
+	if (Region->PerAffectedActorStencilData.Num() > 0)
+	{
+		TMap<TSoftObjectPtr<AActor>, TSharedPtr<FStencilData>> ActorsToUpdate;
+		for (const TPair<TSoftObjectPtr<AActor>, TSharedPtr<FStencilData>>& StencilDataPair : Region->PerAffectedActorStencilData)
+		{
+			TSoftObjectPtr<AActor> Actor = StencilDataPair.Key;
+			if (Actor.IsValid())
+			{
+				TArray<UPrimitiveComponent*> PrimitiveComponents;
 
+				Actor->GetComponents<UPrimitiveComponent>(PrimitiveComponents);
+				for (UPrimitiveComponent* PrimitiveComponent : PrimitiveComponents)
+				{
+					if (PrimitiveComponent->CustomDepthStencilValue != StencilDataPair.Value->AssignedStencil || !PrimitiveComponent->bRenderCustomDepth)
+					{
+						auto Text = FString::Printf(TEXT("Custom Stencil Id for Actor: '%s', Component: '%s' is managed by Color Correct Region '%s' and has been modified manually. Would you like to change it back?"), *Actor->GetName(), *PrimitiveComponent->GetName(), *Region->GetName());
+						EAppReturnType::Type Answer = FMessageDialog::Open(EAppMsgType::YesNo, FText::FromString(Text));
+						if (Answer == EAppReturnType::No)
+						{
+							// Update our stored value and then update other components with the new stencil Id value.
+							StencilDataPair.Value->AssignedStencil = PrimitiveComponent->CustomDepthStencilValue;
+						}
+						ActorsToUpdate.Add(StencilDataPair.Key, StencilDataPair.Value);
+						break;
+					}
+				}
+			}
+		}
+
+		for (const TPair<TSoftObjectPtr<AActor>, TSharedPtr<FStencilData>>& ActorToUpdate : ActorsToUpdate)
+		{
+			Region->PerAffectedActorStencilData.Remove(ActorToUpdate.Key);
+
+			// Assign stencil values without notifying the user.
+			AssignStencilIdToActor(Region, true, false, ActorToUpdate.Key, ActorToUpdate.Value);
+
+		}
+	}
+	
+}
\ No newline at end of file
diff --git a/Engine/Plugins/Experimental/ColorCorrectRegions/Source/ColorCorrectRegions/Public/ColorCorrectRegion.h b/Engine/Plugins/Experimental/ColorCorrectRegions/Source/ColorCorrectRegions/Public/ColorCorrectRegion.h
--- a/Engine/Plugins/Experimental/ColorCorrectRegions/Source/ColorCorrectRegions/Public/ColorCorrectRegion.h
+++ b/Engine/Plugins/Experimental/ColorCorrectRegions/Source/ColorCorrectRegions/Public/ColorCorrectRegion.h
@@ -181,10 +181,10 @@ public:
 		return FirstPrimitiveId;
 	};
 
-	TArray<uint32> ConsumeStencilIds() 
+	TArray<uint32> GetCopyStencilIds() 
 	{ 
 		FScopeLock RegionScopeLock(&StencilIdsCriticalSection);
-		return MoveTemp(StencilIds); 
+		return StencilIds; 
 	};
 private:
 	/**
@@ -213,10 +213,11 @@ private:
 	* A list of actors that need to be excluded or included from color correction to be used on render thread.
 	* This list is populated during engine tick and populated with Stencil Ids from Primitive Components of AffectedActors.
 	* The main reason for this is acessing Primitive Components is not thread safe and is better to be accessed on Game thread.
-	* 
-	* Stencil Ids are moved on render thread to avoid locking either threads for too long.
 	*/
 	TArray<uint32> StencilIds;
 	FCriticalSection StencilIdsCriticalSection;
 
+	// This is for optimization purposes that would let us check assigned actors component's stencil ids ever few once in a while.
+	float TimeWaited = 0;
+
 };
diff --git a/Engine/Plugins/Experimental/ColorCorrectRegions/Source/ColorCorrectRegions/Public/ColorCorrectRegionsStencilManager.h b/Engine/Plugins/Experimental/ColorCorrectRegions/Source/ColorCorrectRegions/Public/ColorCorrectRegionsStencilManager.h
--- a/Engine/Plugins/Experimental/ColorCorrectRegions/Source/ColorCorrectRegions/Public/ColorCorrectRegionsStencilManager.h
+++ b/Engine/Plugins/Experimental/ColorCorrectRegions/Source/ColorCorrectRegions/Public/ColorCorrectRegionsStencilManager.h
@@ -15,4 +15,6 @@ public:
 		static void ClearInvalidActorsForSelectedRegion(AColorCorrectRegion* Region);
 		static void AssignStencilIdsToAllActorsForCCR(UWorld* CurrentWorld, AColorCorrectRegion* Region, bool bIgnoreUserNotifications, bool bSoftAssign);
 
+		/** Runs over all actors assigned to the selected Region and makes sure that Stencil Ids match the ones assigned to it. */
+		static void CheckAssignedActorsValidity(AColorCorrectRegion* Region);
 };
\ No newline at end of file
diff --git a/Engine/Plugins/Experimental/ColorCorrectRegions/Source/ColorCorrectRegions/Public/ColorCorrectRegionsSubsystem.h b/Engine/Plugins/Experimental/ColorCorrectRegions/Source/ColorCorrectRegions/Public/ColorCorrectRegionsSubsystem.h
--- a/Engine/Plugins/Experimental/ColorCorrectRegions/Source/ColorCorrectRegions/Public/ColorCorrectRegionsSubsystem.h
+++ b/Engine/Plugins/Experimental/ColorCorrectRegions/Source/ColorCorrectRegions/Public/ColorCorrectRegionsSubsystem.h
@@ -87,7 +87,10 @@ public:
 	void AssignStencilIdsToPerActorCC(AColorCorrectRegion* Region, bool bIgnoreUserNotificaion = false, bool bSoftAssign = false);
 	
 	/** Handles removal of Stencil Ids for the selected CCR. */
-	void ClearStencilIdsToPerActorCC(AColorCorrectRegion* Region);
+	void ClearStencilIdsToPerActorCC(AColorCorrectRegion* Region);	
+	
+	/** Handles cases when stencil Id has been changed from outside by the user manually. */
+	void CheckAssignedActorsValidity(AColorCorrectRegion* Region);
 
 	/** Resets all stencils and re-assigns for each CCR in the scene. */
 	UFUNCTION(BlueprintCallable, meta = (Category = "Color Correct Regions"))
