Color Correct Regions:
- Changed the icon in editor for CCW.
- Fixed the issue for Circle CCW which was causing artifacts on intersection.
- CCW no longer has collision in play mode.

diff --git a/Engine/Plugins/Experimental/ColorCorrectRegions/Shaders/Private/ColorCorrectRegionsShader.usf b/Engine/Plugins/Experimental/ColorCorrectRegions/Shaders/Private/ColorCorrectRegionsShader.usf
--- a/Engine/Plugins/Experimental/ColorCorrectRegions/Shaders/Private/ColorCorrectRegionsShader.usf
+++ b/Engine/Plugins/Experimental/ColorCorrectRegions/Shaders/Private/ColorCorrectRegionsShader.usf
@@ -179,7 +179,7 @@ float GetShapeAlpha(in out FMaterialPixelParameters Parameters, in FExternalExpr
 	float2 Interpolator = smoothstep(ExternalExpressions.Outer, ExternalExpressions.Inner, abs(PositionOnPlane.xy)) * ClipZ * step(IntersectionSide, 0.);
 	float ShapeAlpha = Interpolator.x * step(0., Interpolator.y) * Interpolator.g * step(0., Interpolator.x);
 #elif WINDOW_TYPE == WINDOW_TYPE_CIRCLE
-	float ShapeAlpha = smoothstep(ExternalExpressions.Outer, ExternalExpressions.Inner, length(PositionOnPlane)) * ClipZ;
+	float ShapeAlpha = smoothstep(ExternalExpressions.Outer, ExternalExpressions.Inner, length(PositionOnPlane)) * ClipZ * step(IntersectionSide, 0.);
 #endif
 
 #if CLIP_PIXELS_OUTSIDE_AABB
diff --git a/Engine/Plugins/Experimental/ColorCorrectRegions/Source/ColorCorrectRegions/Private/ColorCorrectRegion.cpp b/Engine/Plugins/Experimental/ColorCorrectRegions/Source/ColorCorrectRegions/Private/ColorCorrectRegion.cpp
--- a/Engine/Plugins/Experimental/ColorCorrectRegions/Source/ColorCorrectRegions/Private/ColorCorrectRegion.cpp
+++ b/Engine/Plugins/Experimental/ColorCorrectRegions/Source/ColorCorrectRegions/Private/ColorCorrectRegion.cpp
@@ -1,6 +1,7 @@
 // Copyright Epic Games, Inc. All Rights Reserved.
 
 #include "ColorCorrectRegion.h"
+#include "ColorCorrectWindow.h"
 #include "Components/BillboardComponent.h"
 #include "ColorCorrectRegionsSubsystem.h"
 #include "UObject/ConstructorHelpers.h"
@@ -34,48 +35,12 @@
 	RootComponent = ObjectInitializer.CreateDefaultSubobject<USceneComponent>(this, TEXT("Root"));
 	RootComponent->SetMobility(EComponentMobility::Movable);
 
-#if WITH_EDITOR
-
-	// Create billboard component
-
-	if (GIsEditor && !IsRunningCommandlet())
+#if WITH_METADATA
+	if (!Cast<AColorCorrectWindow>(this))
 	{
-		// Structure to hold one-time initialization
-	
-		struct FConstructorStatics
-		{
-			ConstructorHelpers::FObjectFinderOptional<UTexture2D> SpriteTextureObject;
-			FName ID_ColorCorrectRegion;
-			FText NAME_ColorCorrectRegion;
-
-			FConstructorStatics()
-				: SpriteTextureObject(TEXT("/ColorCorrectRegions/Icons/S_ColorCorrectRegionIcon"))
-				, ID_ColorCorrectRegion(TEXT("Color Correct Region"))
-				, NAME_ColorCorrectRegion(NSLOCTEXT("SpriteCategory", "ColorCorrectRegion", "Color Correct Region"))
-			{
-			}
-		};
-
-		static FConstructorStatics ConstructorStatics;
-
-		SpriteComponent = ObjectInitializer.CreateEditorOnlyDefaultSubobject<UBillboardComponent>(this, TEXT("Color Correct Region Icon"));
-
-		if (SpriteComponent)
-		{
-			SpriteComponent->Sprite = ConstructorStatics.SpriteTextureObject.Get();
-			SpriteComponent->SpriteInfo.Category = ConstructorStatics.ID_ColorCorrectRegion;
-			SpriteComponent->SpriteInfo.DisplayName = ConstructorStatics.NAME_ColorCorrectRegion;
-			SpriteComponent->SetIsVisualizationComponent(true);
-			SpriteComponent->SetRelativeLocationAndRotation(FVector::ZeroVector, FRotator::ZeroRotator);
-			SpriteComponent->SetMobility(EComponentMobility::Movable);
-			SpriteComponent->bHiddenInGame = true;
-			SpriteComponent->bIsScreenSizeScaled = true;
-
-			SpriteComponent->AttachToComponent(RootComponent, FAttachmentTransformRules::KeepRelativeTransform);
-		}
+		CreateIcon();
 	}
-
-#endif // WITH_EDITOR
+#endif
 }
 
 void AColorCorrectRegion::BeginPlay()
@@ -227,6 +192,50 @@
 	}
 }
 
+#if WITH_METADATA
+void AColorCorrectRegion::CreateIcon()
+{
+	// Create billboard component
+	if (GIsEditor && !IsRunningCommandlet())
+	{
+		// Structure to hold one-time initialization
+
+		struct FConstructorStatics
+		{
+			ConstructorHelpers::FObjectFinderOptional<UTexture2D> SpriteTextureObject;
+			FName ID_ColorCorrectRegion;
+			FText NAME_ColorCorrectRegion;
+
+			FConstructorStatics()
+				: SpriteTextureObject(TEXT("/ColorCorrectRegions/Icons/S_ColorCorrectRegionIcon"))
+				, ID_ColorCorrectRegion(TEXT("Color Correct Region"))
+				, NAME_ColorCorrectRegion(NSLOCTEXT("SpriteCategory", "ColorCorrectRegion", "Color Correct Region"))
+			{
+			}
+		};
+
+		static FConstructorStatics ConstructorStatics;
+
+		SpriteComponent = CreateEditorOnlyDefaultSubobject<UBillboardComponent>(TEXT("Color Correct Region Icon"));
+
+		if (SpriteComponent)
+		{
+			SpriteComponent->Sprite = ConstructorStatics.SpriteTextureObject.Get();
+			SpriteComponent->SpriteInfo.Category = ConstructorStatics.ID_ColorCorrectRegion;
+			SpriteComponent->SpriteInfo.DisplayName = ConstructorStatics.NAME_ColorCorrectRegion;
+			SpriteComponent->SetIsVisualizationComponent(true);
+			SpriteComponent->SetRelativeLocationAndRotation(FVector::ZeroVector, FRotator::ZeroRotator);
+			SpriteComponent->SetMobility(EComponentMobility::Movable);
+			SpriteComponent->bHiddenInGame = true;
+			SpriteComponent->bIsScreenSizeScaled = true;
+
+			SpriteComponent->AttachToComponent(RootComponent, FAttachmentTransformRules::KeepRelativeTransform);
+		}
+	}
+
+}
+#endif 
+
 #if WITH_EDITOR
 void AColorCorrectRegion::PostEditChangeProperty(struct FPropertyChangedEvent& PropertyChangedEvent)
 {
diff --git a/Engine/Plugins/Experimental/ColorCorrectRegions/Source/ColorCorrectRegions/Private/ColorCorrectWindow.cpp b/Engine/Plugins/Experimental/ColorCorrectRegions/Source/ColorCorrectRegions/Private/ColorCorrectWindow.cpp
--- a/Engine/Plugins/Experimental/ColorCorrectRegions/Source/ColorCorrectRegions/Private/ColorCorrectWindow.cpp
+++ b/Engine/Plugins/Experimental/ColorCorrectRegions/Source/ColorCorrectRegions/Private/ColorCorrectWindow.cpp
@@ -1,6 +1,7 @@
 // Copyright Epic Games, Inc. All Rights Reserved.
 
 #include "ColorCorrectWindow.h"
+#include "Components/BillboardComponent.h"
 #include "CoreMinimal.h"
 #include "Materials/Material.h"
 #include "Engine/StaticMesh.h"
@@ -29,8 +30,14 @@ AColorCorrectWindow::AColorCorrectWindow(const FObjectInitializer& ObjectInitial
 		MeshComponent->SetStaticMesh(StaticMeshes[static_cast<uint8>(CCWType)]);
 		MeshComponent->SetWorldScale3D(FVector(1., 1., 0.001));
 		MeshComponent->SetMaterial(0, Material);
+		MeshComponent->SetCollisionEnabled(ECollisionEnabled::QueryOnly);
+		MeshComponent->SetCollisionProfileName(TEXT("OverlapAll"));
 	}
 	SetMeshVisibilityForWindowType();
+
+#if WITH_METADATA
+	CreateIcon();
+#endif
 }
 
 #if WITH_EDITOR
@@ -46,6 +53,51 @@ void AColorCorrectWindow::PostEditChangeProperty(struct FPropertyChangedEvent& P
 }
 #endif //WITH_EDITOR
 
+
+#if WITH_METADATA
+void AColorCorrectWindow::CreateIcon()
+{
+	// Create billboard component
+	if (GIsEditor && !IsRunningCommandlet())
+	{
+		// Structure to hold one-time initialization
+
+		struct FConstructorStatics
+		{
+			ConstructorHelpers::FObjectFinderOptional<UTexture2D> SpriteTextureObject;
+			FName ID_ColorCorrectRegion;
+			FText NAME_ColorCorrectRegion;
+
+			FConstructorStatics()
+				: SpriteTextureObject(TEXT("/ColorCorrectRegions/Icons/S_ColorCorrectWindowIcon"))
+				, ID_ColorCorrectRegion(TEXT("Color Correct Window"))
+				, NAME_ColorCorrectRegion(NSLOCTEXT("SpriteCategory", "ColorCorrectWindow", "Color Correct Window"))
+			{
+			}
+		};
+
+		static FConstructorStatics ConstructorStatics;
+
+		SpriteComponent = CreateEditorOnlyDefaultSubobject<UBillboardComponent>(TEXT("Color Correct Window Icon"));
+
+		if (SpriteComponent)
+		{
+			SpriteComponent->Sprite = ConstructorStatics.SpriteTextureObject.Get();
+			SpriteComponent->SpriteInfo.Category = ConstructorStatics.ID_ColorCorrectRegion;
+			SpriteComponent->SpriteInfo.DisplayName = ConstructorStatics.NAME_ColorCorrectRegion;
+			SpriteComponent->SetIsVisualizationComponent(true);
+			SpriteComponent->SetRelativeLocationAndRotation(FVector::ZeroVector, FRotator::ZeroRotator);
+			SpriteComponent->SetMobility(EComponentMobility::Movable);
+			SpriteComponent->bHiddenInGame = true;
+			SpriteComponent->bIsScreenSizeScaled = true;
+
+			SpriteComponent->AttachToComponent(RootComponent, FAttachmentTransformRules::KeepRelativeTransform);
+		}
+	}
+
+}
+#endif 
+
 void AColorCorrectWindow::SetMeshVisibilityForWindowType()
 {
 	for (EColorCorrectWindowType CCWType : TEnumRange<EColorCorrectWindowType>())
diff --git a/Engine/Plugins/Experimental/ColorCorrectRegions/Source/ColorCorrectRegions/Public/ColorCorrectRegion.h b/Engine/Plugins/Experimental/ColorCorrectRegions/Source/ColorCorrectRegions/Public/ColorCorrectRegion.h
--- a/Engine/Plugins/Experimental/ColorCorrectRegions/Source/ColorCorrectRegions/Public/ColorCorrectRegion.h
+++ b/Engine/Plugins/Experimental/ColorCorrectRegions/Source/ColorCorrectRegions/Public/ColorCorrectRegion.h
@@ -187,6 +187,13 @@ public:
 		return StencilIds; 
 	};
 private:
+
+#if WITH_METADATA
+	/** Creates an icon for CCR/CCW to be clicked on in Editor. */
+	void CreateIcon();
+
+#endif // WITH_METADATA
+
 	/**
 	* AffectedActors property change could potentially invoke a Dialog Window, which should be displayed on Game Thread.
 	*/
diff --git a/Engine/Plugins/Experimental/ColorCorrectRegions/Source/ColorCorrectRegions/Public/ColorCorrectWindow.h b/Engine/Plugins/Experimental/ColorCorrectRegions/Source/ColorCorrectRegions/Public/ColorCorrectWindow.h
--- a/Engine/Plugins/Experimental/ColorCorrectRegions/Source/ColorCorrectRegions/Public/ColorCorrectWindow.h
+++ b/Engine/Plugins/Experimental/ColorCorrectRegions/Source/ColorCorrectRegions/Public/ColorCorrectWindow.h
@@ -44,6 +44,11 @@ public:
 #endif
 
 private:
+#if WITH_METADATA
+	void CreateIcon();
+#endif // WITH_METADATA
+
+	/** Swaps meshes for different CCW. */
 	void SetMeshVisibilityForWindowType();
 
 };
