Color Correct Regions:
- Added transactions for managed objects so that changes are reflected in nDisplay in real time.
- Added CCR deletion handling.

diff --git a/Engine/Plugins/Experimental/ColorCorrectRegions/Source/ColorCorrectRegions/Private/ColorCorrectRegionSubsystem.cpp b/Engine/Plugins/Experimental/ColorCorrectRegions/Source/ColorCorrectRegions/Private/ColorCorrectRegionSubsystem.cpp
--- a/Engine/Plugins/Experimental/ColorCorrectRegions/Source/ColorCorrectRegions/Private/ColorCorrectRegionSubsystem.cpp
+++ b/Engine/Plugins/Experimental/ColorCorrectRegions/Source/ColorCorrectRegions/Private/ColorCorrectRegionSubsystem.cpp
@@ -15,6 +15,8 @@
 #include "Editor.h"
 #endif
 
+#define LOCTEXT_NAMESPACE "CCR"
+
 namespace
 {
 	bool IsRegionValid(AColorCorrectRegion* InRegion, UWorld* CurrentWorld)
@@ -124,6 +126,11 @@ void UColorCorrectRegionsSubsystem::OnActorDeleted(AActor* InActor)
 #endif
 	{
 		AsRegion->Cleanup();
+
+#if WITH_EDITOR
+		FColorCorrectRegionsStencilManager::OnCCRRemoved(GetWorld(), AsRegion);
+#endif
+
 		FScopeLock RegionScopeLock(&RegionAccessCriticalSection);
 		RegionsPriorityBased.Remove(AsRegion);
 		RegionsDistanceBased.Remove(AsRegion);
@@ -161,12 +168,41 @@ void UColorCorrectRegionsSubsystem::SortRegionsByDistance(const FVector& ViewLoc
 
 void UColorCorrectRegionsSubsystem::AssignStencilIdsToPerActorCC(AColorCorrectRegion* Region, bool bIgnoreUserNotificaion, bool bSoftAssign)
 {
+#if WITH_EDITOR
+	if (!bSoftAssign && GEditor)
+	{
+		GEditor->BeginTransaction(LOCTEXT("PerActorCCActorAssigned", "Per actor CC Actor Assigned"));
+	}
+#endif
 	FColorCorrectRegionsStencilManager::AssignStencilIdsToAllActorsForCCR(GetWorld(), Region, bIgnoreUserNotificaion, bSoftAssign);
+
+#if WITH_EDITOR
+	if (!bSoftAssign && GEditor)
+	{
+		this->Modify();
+		GEditor->EndTransaction();
+	}
+#endif
 }
 
 void UColorCorrectRegionsSubsystem::ClearStencilIdsToPerActorCC(AColorCorrectRegion* Region)
 {
+#if WITH_EDITOR
+	if (GEditor)
+	{
+		GEditor->BeginTransaction(LOCTEXT("PerActorCCActorRemoved", "Per actor CC Actor Removed"));
+	}
+#endif
+
 	FColorCorrectRegionsStencilManager::RemoveStencilNumberForSelectedRegion(GetWorld(), Region);
+	this->Modify();
+
+#if WITH_EDITOR
+	if (GEditor)
+	{
+		GEditor->EndTransaction();
+	}
+#endif
 }
 
 void UColorCorrectRegionsSubsystem::CheckAssignedActorsValidity(AColorCorrectRegion* Region)
@@ -217,3 +253,4 @@ void UColorCorrectRegionsSubsystem::RefreshRegions()
 	RefreshStenciIdAssignmentForAllCCR();
 }
 
+#undef LOCTEXT_NAMESPACE
\ No newline at end of file
diff --git a/Engine/Plugins/Experimental/ColorCorrectRegions/Source/ColorCorrectRegions/Private/ColorCorrectRegionsStencilManager.cpp b/Engine/Plugins/Experimental/ColorCorrectRegions/Source/ColorCorrectRegions/Private/ColorCorrectRegionsStencilManager.cpp
--- a/Engine/Plugins/Experimental/ColorCorrectRegions/Source/ColorCorrectRegions/Private/ColorCorrectRegionsStencilManager.cpp
+++ b/Engine/Plugins/Experimental/ColorCorrectRegions/Source/ColorCorrectRegions/Private/ColorCorrectRegionsStencilManager.cpp
@@ -68,7 +68,12 @@ namespace
 	{
 		TArray<UPrimitiveComponent*> PrimitiveComponents;
 		ActorToAssignStencilTo->GetComponents<UPrimitiveComponent>(PrimitiveComponents);
-
+#if WITH_EDITOR
+		if (!bSoftAssign)
+		{
+			ActorToAssignStencilTo->Modify();
+		}
+#endif
 		uint8 StencilIdMin = FMath::Min(CVarStencilIdRangeMin.GetValueOnAnyThread(), CVarStencilIdRangeMax.GetValueOnAnyThread());
 		uint8 StencilIdMax = FMath::Max(CVarStencilIdRangeMin.GetValueOnAnyThread(), CVarStencilIdRangeMax.GetValueOnAnyThread());
 		for (UPrimitiveComponent* PrimitiveComponent : PrimitiveComponents)
@@ -97,6 +102,9 @@ namespace
 				PrimitiveComponent->SetRenderCustomDepth(true);
 				PrimitiveComponent->CustomDepthStencilValue = static_cast<int32>(StencilData->AssignedStencil);
 				PrimitiveComponent->MarkRenderStateDirty();
+#if WITH_EDITOR
+				ActorToAssignStencilTo->Modify();
+#endif
 			}
 			else
 			{
@@ -116,7 +124,9 @@ namespace
 		{
 			TArray<UPrimitiveComponent*> PrimitiveComponents;
 			ActorToAssignStencilTo->GetComponents<UPrimitiveComponent>(PrimitiveComponents);
-
+#if WITH_EDITOR
+			ActorToAssignStencilTo->Modify();
+#endif
 			for (UPrimitiveComponent* PrimitiveComponent : PrimitiveComponents)
 			{
 				// If the stencil id is within our range then it is managed by us, so we should reset this to default.
@@ -125,6 +135,9 @@ namespace
 					PrimitiveComponent->SetRenderCustomDepth(false);
 					PrimitiveComponent->CustomDepthStencilValue = 0;
 					PrimitiveComponent->MarkRenderStateDirty();
+#if WITH_EDITOR
+					PrimitiveComponent->Modify();
+#endif
 				}
 			}
 		}
@@ -207,7 +220,7 @@ void FColorCorrectRegionsStencilManager::AssignStencilNumberToActorForSelectedRe
 
 		if (Region->PerAffectedActorStencilData.Num() > 0)
 		{
-			for (TPair<TSoftObjectPtr<AActor>, TSharedPtr<FStencilData>> ActorDataPair : Region->PerAffectedActorStencilData)
+			for (const TPair<TSoftObjectPtr<AActor>, TSharedPtr<FStencilData>>& ActorDataPair : Region->PerAffectedActorStencilData)
 			{
 				if (!IsActorAssignedToOtherCCR(Region, ActorDataPair.Key, AllCCRsInCurrentLevel))
 				{
@@ -263,7 +276,7 @@ void FColorCorrectRegionsStencilManager::RemoveStencilNumberForSelectedRegion(UW
 	TSet<TSoftObjectPtr<AActor>> ActorsToCleanup;
 
 	// Accumulate all CCRs that were removed.
-	for (TPair<TSoftObjectPtr<AActor>, TSharedPtr<FStencilData>> ActorDataPair : Region->PerAffectedActorStencilData)
+	for (const TPair<TSoftObjectPtr<AActor>, TSharedPtr<FStencilData>>& ActorDataPair : Region->PerAffectedActorStencilData)
 	{
 		if (!Region->AffectedActors.Contains(ActorDataPair.Key))
 		{
@@ -301,6 +314,41 @@ void FColorCorrectRegionsStencilManager::RemoveStencilNumberForSelectedRegion(UW
 	}
 }
 
+void FColorCorrectRegionsStencilManager::OnCCRRemoved(UWorld* CurrentWorld, AColorCorrectRegion* Region)
+{
+	ULevel* CurrentLevel = CurrentWorld->GetCurrentLevel();
+	TArray<AActor*> AllCCRsInCurrentLevel = CurrentLevel->Actors.FilterByPredicate([](const AActor* Actor) { return Cast<AColorCorrectRegion>(Actor) != nullptr; });
+
+	// For each actor that was removed, remove Ids if it is unused by any other CCR, or otherwise assign a new id.
+	for (const TPair<TSoftObjectPtr<AActor>, TSharedPtr<FStencilData>>& ActorDataPair : Region->PerAffectedActorStencilData)
+	{
+		bool bUsedByAnotherActor = false;
+		for (AActor* ActorCCR : AllCCRsInCurrentLevel)
+		{
+			AColorCorrectRegion* OtherCCR = Cast<AColorCorrectRegion>(ActorCCR);
+			if (OtherCCR == Region)
+			{
+				continue;
+			}
+
+			if (OtherCCR->PerAffectedActorStencilData.Contains(ActorDataPair.Key))
+			{
+				OtherCCR->PerAffectedActorStencilData.Remove(ActorDataPair.Key);
+				AssignStencilNumberToActorForSelectedRegion(CurrentWorld, OtherCCR, ActorDataPair.Key, true, false);
+				bUsedByAnotherActor = true;
+				break;
+			}
+		}
+
+		// This actor is not used by any aother actor. Remove it.
+		if (!bUsedByAnotherActor)
+		{
+			ClearStencilIdFromActor(ActorDataPair.Key);
+		}
+	}
+	Region->PerAffectedActorStencilData.Empty();
+}
+
 void FColorCorrectRegionsStencilManager::ClearInvalidActorsForSelectedRegion(AColorCorrectRegion* Region)
 {
 	TArray<TSoftObjectPtr<AActor>> ActorsToCleanup;
@@ -350,7 +398,7 @@ void FColorCorrectRegionsStencilManager::CheckAssignedActorsValidity(AColorCorre
 				{
 					if (PrimitiveComponent->CustomDepthStencilValue != StencilDataPair.Value->AssignedStencil || !PrimitiveComponent->bRenderCustomDepth)
 					{
-						auto Text = FString::Printf(TEXT("Custom Stencil Id for Actor: '%s', Component: '%s' is managed by Color Correct Region '%s' and has been modified manually. Would you like to change it back?"), *Actor->GetName(), *PrimitiveComponent->GetName(), *Region->GetName());
+						auto Text = FString::Printf(TEXT("Custom Stencil Id for Actor: '%s', Component: '%s' is managed by Color Correct Region '%s' and has been modified manually. Would you like to change the id back?"), *Actor->GetName(), *PrimitiveComponent->GetName(), *Region->GetName());
 						EAppReturnType::Type Answer = FMessageDialog::Open(EAppMsgType::YesNo, FText::FromString(Text));
 						if (Answer == EAppReturnType::No)
 						{
diff --git a/Engine/Plugins/Experimental/ColorCorrectRegions/Source/ColorCorrectRegions/Public/ColorCorrectRegionsStencilManager.h b/Engine/Plugins/Experimental/ColorCorrectRegions/Source/ColorCorrectRegions/Public/ColorCorrectRegionsStencilManager.h
--- a/Engine/Plugins/Experimental/ColorCorrectRegions/Source/ColorCorrectRegions/Public/ColorCorrectRegionsStencilManager.h
+++ b/Engine/Plugins/Experimental/ColorCorrectRegions/Source/ColorCorrectRegions/Public/ColorCorrectRegionsStencilManager.h
@@ -12,6 +12,7 @@ class FColorCorrectRegionsStencilManager
 public:
 		static void AssignStencilNumberToActorForSelectedRegion(UWorld* CurrentWorld, AColorCorrectRegion* Region, TSoftObjectPtr<AActor> ActorToAssignStencilTo, bool bIgnoreUserNotifications, bool bSoftAssign);
 		static void RemoveStencilNumberForSelectedRegion(UWorld* CurrentWorld, AColorCorrectRegion* Region);
+		static void OnCCRRemoved(UWorld* CurrentWorld, AColorCorrectRegion* Region);
 		static void ClearInvalidActorsForSelectedRegion(AColorCorrectRegion* Region);
 		static void AssignStencilIdsToAllActorsForCCR(UWorld* CurrentWorld, AColorCorrectRegion* Region, bool bIgnoreUserNotifications, bool bSoftAssign);
 
