Color Correct Regions: Notifying user if Custom Depth mode is not set to Enabled with Stencil.

diff --git a/Engine/Plugins/Experimental/ColorCorrectRegions/Source/ColorCorrectRegions/Private/ColorCorrectRegionsSceneViewExtension.cpp b/Engine/Plugins/Experimental/ColorCorrectRegions/Source/ColorCorrectRegions/Private/ColorCorrectRegionsSceneViewExtension.cpp
--- a/Engine/Plugins/Experimental/ColorCorrectRegions/Source/ColorCorrectRegions/Private/ColorCorrectRegionsSceneViewExtension.cpp
+++ b/Engine/Plugins/Experimental/ColorCorrectRegions/Source/ColorCorrectRegions/Private/ColorCorrectRegionsSceneViewExtension.cpp
@@ -247,6 +247,21 @@ namespace
 		, const TArray<uint32>& StencilIds
 		, FScreenPassRenderTarget& OutMergedStencilRenderTarget)
 	{
+		static bool bNotifiedAboutCustomDepth = false;
+		static const auto CVarCustomDepth = IConsoleManager::Get().FindTConsoleVariableDataInt(TEXT("r.CustomDepth"));
+		const int32 EnabledWithStencil = 3;
+
+		if (CVarCustomDepth->GetValueOnAnyThread() != EnabledWithStencil && !bNotifiedAboutCustomDepth)
+		{
+			UE_LOG(ColorCorrectRegions, Error, TEXT("Per Actor Color Correction requires Custom Depth Mode to be set to \"Enabled With Stencil\""));
+			bNotifiedAboutCustomDepth = true;
+			return;
+		}
+		else if (CVarCustomDepth->GetValueOnAnyThread() == EnabledWithStencil)
+		{
+			bNotifiedAboutCustomDepth = false;
+		}
+
 		if (StencilIds.Num() == 0)
 		{
 			return;
