Color Correct Regions:
- Deprecated old CCR blueprint. This will remove duplicate blueprint from place actors and allows us to rename color correct regions to color correction regions.
- Renamed Color Correct Windows to Color Correction windows

diff --git a/Engine/Plugins/Experimental/ColorCorrectRegions/ColorCorrectRegions.uplugin b/Engine/Plugins/Experimental/ColorCorrectRegions/ColorCorrectRegions.uplugin
--- a/Engine/Plugins/Experimental/ColorCorrectRegions/ColorCorrectRegions.uplugin
+++ b/Engine/Plugins/Experimental/ColorCorrectRegions/ColorCorrectRegions.uplugin
@@ -2,7 +2,7 @@
 	"FileVersion": 3,
 	"Version": 1,
 	"VersionName": "1.0",
-	"FriendlyName": "Color Correct Regions (CCR)",
+	"FriendlyName": "Color Correction Regions (CCR)",
 	"Description": "Color correction/shading constrained to regions/volumes",
 	"Category": "Other",
 	"CreatedBy": "",
diff --git a/Engine/Plugins/Experimental/ColorCorrectRegions/Source/ColorCorrectRegions/Private/ColorCorrectRegion.cpp b/Engine/Plugins/Experimental/ColorCorrectRegions/Source/ColorCorrectRegions/Private/ColorCorrectRegion.cpp
--- a/Engine/Plugins/Experimental/ColorCorrectRegions/Source/ColorCorrectRegions/Private/ColorCorrectRegion.cpp
+++ b/Engine/Plugins/Experimental/ColorCorrectRegions/Source/ColorCorrectRegions/Private/ColorCorrectRegion.cpp
@@ -8,9 +8,13 @@
 #include "CoreMinimal.h"
 #include "Engine/GameEngine.h"
 #include "Engine/Classes/Components/MeshComponent.h"
+#include "Materials/Material.h"
+#include "Engine/StaticMesh.h"
 #include "Components/StaticMeshComponent.h"
 #include "Engine/Texture2D.h"
 
+ENUM_RANGE_BY_COUNT(EColorCorrectRegionsType, EColorCorrectRegionsType::MAX)
+
 
 AColorCorrectRegion::AColorCorrectRegion(const FObjectInitializer& ObjectInitializer) 
 	: Super(ObjectInitializer)
@@ -36,7 +40,7 @@ AColorCorrectRegion::AColorCorrectRegion(const FObjectInitializer& ObjectInitial
 	RootComponent->SetMobility(EComponentMobility::Movable);
 
 #if WITH_METADATA
-	if (!Cast<AColorCorrectWindow>(this))
+	if (!Cast<AColorCorrectionWindow>(this))
 	{
 		CreateIcon();
 	}
@@ -278,3 +282,60 @@ void AColorCorrectRegion::PostEditChangeProperty(struct FPropertyChangedEvent& P
 }
 #endif //WITH_EDITOR
 
+
+AColorCorrectionRegion::AColorCorrectionRegion(const FObjectInitializer& ObjectInitializer)
+	: Super(ObjectInitializer)
+{
+	UMaterial* Material = LoadObject<UMaterial>(NULL, TEXT("/ColorCorrectRegions/Materials/M_ColorCorrectRegionTransparentPreview.M_ColorCorrectRegionTransparentPreview"), NULL, LOAD_None, NULL);
+	const TArray<UStaticMesh*> StaticMeshes =
+	{
+		Cast<UStaticMesh>(StaticLoadObject(UStaticMesh::StaticClass(), NULL, TEXT("/Engine/BasicShapes/Sphere"))),
+		Cast<UStaticMesh>(StaticLoadObject(UStaticMesh::StaticClass(), NULL, TEXT("/Engine/BasicShapes/Cube"))),
+		Cast<UStaticMesh>(StaticLoadObject(UStaticMesh::StaticClass(), NULL, TEXT("/Engine/BasicShapes/Cylinder"))),
+		Cast<UStaticMesh>(StaticLoadObject(UStaticMesh::StaticClass(), NULL, TEXT("/Engine/BasicShapes/Cone")))
+	};
+
+	for (EColorCorrectRegionsType CCRType : TEnumRange<EColorCorrectRegionsType>())
+	{
+		UStaticMeshComponent* MeshComponent = CreateDefaultSubobject<UStaticMeshComponent>(*UEnum::GetValueAsString(CCRType));
+		MeshComponents.Add(MeshComponent);
+		MeshComponent->SetupAttachment(RootComponent);
+		MeshComponent->SetStaticMesh(StaticMeshes[static_cast<uint8>(CCRType)]);
+		MeshComponent->SetMaterial(0, Material);
+		MeshComponent->SetCollisionEnabled(ECollisionEnabled::QueryOnly);
+		MeshComponent->SetCollisionProfileName(TEXT("OverlapAll"));
+	}
+
+	SetMeshVisibilityForRegionType();
+
+}
+
+void AColorCorrectionRegion::SetMeshVisibilityForRegionType()
+{
+	for (EColorCorrectRegionsType CCRType : TEnumRange<EColorCorrectRegionsType>())
+	{
+		uint8 TypeIndex = static_cast<uint8>(CCRType);
+
+		if (CCRType == Type)
+		{
+			MeshComponents[TypeIndex]->SetVisibility(true, true);
+		}
+		else
+		{
+			MeshComponents[TypeIndex]->SetVisibility(false, true);
+		}
+	}
+}
+
+#if WITH_EDITOR
+void AColorCorrectionRegion::PostEditChangeProperty(struct FPropertyChangedEvent& PropertyChangedEvent)
+{
+	Super::PostEditChangeProperty(PropertyChangedEvent);
+	const FName PropertyName = PropertyChangedEvent.GetPropertyName();
+
+	if (PropertyName == GET_MEMBER_NAME_CHECKED(AColorCorrectionRegion, Type))
+	{
+		SetMeshVisibilityForRegionType();
+	}
+}
+#endif
diff --git a/Engine/Plugins/Experimental/ColorCorrectRegions/Source/ColorCorrectRegions/Private/ColorCorrectRegionSubsystem.cpp b/Engine/Plugins/Experimental/ColorCorrectRegions/Source/ColorCorrectRegions/Private/ColorCorrectRegionSubsystem.cpp
--- a/Engine/Plugins/Experimental/ColorCorrectRegions/Source/ColorCorrectRegions/Private/ColorCorrectRegionSubsystem.cpp
+++ b/Engine/Plugins/Experimental/ColorCorrectRegions/Source/ColorCorrectRegions/Private/ColorCorrectRegionSubsystem.cpp
@@ -101,7 +101,7 @@ void UColorCorrectRegionsSubsystem::OnActorSpawned(AActor* InActor)
 		// We wouldn't have to do a check here except in case of nDisplay we need to populate this list during OnLevelsChanged 
 		// because nDisplay can release Actors while those are marked as BeginningPlay. Therefore we want to avoid 
 		// adding regions twice.
-		bool bIsDistanceBased = Cast<AColorCorrectWindow>(InActor) != nullptr;
+		bool bIsDistanceBased = Cast<AColorCorrectionWindow>(InActor) != nullptr;
 		TArray<AColorCorrectRegion*>* RegionsToAddTo = bIsDistanceBased ? &RegionsDistanceBased : &RegionsPriorityBased;
 		if (!RegionsToAddTo->Contains(AsRegion))
 		{
@@ -237,7 +237,7 @@ void UColorCorrectRegionsSubsystem::RefreshRegions()
 		AColorCorrectRegion* AsRegion = *It;
 		if (IsRegionValid(AsRegion, GetWorld()))
 		{
-			if (!Cast<AColorCorrectWindow>(AsRegion))
+			if (!Cast<AColorCorrectionWindow>(AsRegion))
 			{
 				RegionsPriorityBased.Add(AsRegion);
 			}
diff --git a/Engine/Plugins/Experimental/ColorCorrectRegions/Source/ColorCorrectRegions/Private/ColorCorrectRegionsSceneViewExtension.cpp b/Engine/Plugins/Experimental/ColorCorrectRegions/Source/ColorCorrectRegions/Private/ColorCorrectRegionsSceneViewExtension.cpp
--- a/Engine/Plugins/Experimental/ColorCorrectRegions/Source/ColorCorrectRegions/Private/ColorCorrectRegionsSceneViewExtension.cpp
+++ b/Engine/Plugins/Experimental/ColorCorrectRegions/Source/ColorCorrectRegions/Private/ColorCorrectRegionsSceneViewExtension.cpp
@@ -469,7 +469,7 @@ namespace
 		}
 
 		TShaderRef<FColorCorrectGenericPS> PixelShader;
-		if (AColorCorrectWindow* CCWindow = Cast<AColorCorrectWindow>(Region))
+		if (AColorCorrectionWindow* CCWindow = Cast<AColorCorrectionWindow>(Region))
 		{
 			PixelShader = GetWindowShader(GlobalShaderMap, CCWindow->WindowType, TemperatureType, bIsAdvanced, MergedStencilRenderTarget.IsValid());
 		}
diff --git a/Engine/Plugins/Experimental/ColorCorrectRegions/Source/ColorCorrectRegions/Private/ColorCorrectWindow.cpp b/Engine/Plugins/Experimental/ColorCorrectRegions/Source/ColorCorrectRegions/Private/ColorCorrectWindow.cpp
--- a/Engine/Plugins/Experimental/ColorCorrectRegions/Source/ColorCorrectRegions/Private/ColorCorrectWindow.cpp
+++ b/Engine/Plugins/Experimental/ColorCorrectRegions/Source/ColorCorrectRegions/Private/ColorCorrectWindow.cpp
@@ -12,7 +12,7 @@
 
 ENUM_RANGE_BY_COUNT(EColorCorrectWindowType, EColorCorrectWindowType::MAX)
 
-AColorCorrectWindow::AColorCorrectWindow(const FObjectInitializer& ObjectInitializer)
+AColorCorrectionWindow::AColorCorrectionWindow(const FObjectInitializer& ObjectInitializer)
 	: Super(ObjectInitializer)
 	, WindowType(EColorCorrectWindowType::Square)
 {
@@ -50,12 +50,12 @@ AColorCorrectWindow::AColorCorrectWindow(const FObjectInitializer& ObjectInitial
 }
 
 #if WITH_EDITOR
-void AColorCorrectWindow::PostEditChangeProperty(struct FPropertyChangedEvent& PropertyChangedEvent)
+void AColorCorrectionWindow::PostEditChangeProperty(struct FPropertyChangedEvent& PropertyChangedEvent)
 {
 	Super::PostEditChangeProperty(PropertyChangedEvent);
 	const FName PropertyName = PropertyChangedEvent.GetPropertyName();
 
-	if (PropertyName == GET_MEMBER_NAME_CHECKED(AColorCorrectWindow, WindowType))
+	if (PropertyName == GET_MEMBER_NAME_CHECKED(AColorCorrectionWindow, WindowType))
 	{
 		SetMeshVisibilityForWindowType();
 	}
@@ -82,7 +82,7 @@ void AColorCorrectWindow::PostEditChangeProperty(struct FPropertyChangedEvent& P
 	}
 }
 
-void AColorCorrectWindow::PostEditMove(bool bFinished)
+void AColorCorrectionWindow::PostEditMove(bool bFinished)
 {
 	Super::PostEditMove(bFinished);
 
@@ -94,7 +94,7 @@ void AColorCorrectWindow::PostEditMove(bool bFinished)
 
 
 #if WITH_METADATA
-void AColorCorrectWindow::CreateIcon()
+void AColorCorrectionWindow::CreateIcon()
 {
 	// Create billboard component
 	if (GIsEditor && !IsRunningCommandlet())
@@ -137,7 +137,7 @@ void AColorCorrectWindow::CreateIcon()
 }
 #endif 
 
-void AColorCorrectWindow::SetMeshVisibilityForWindowType()
+void AColorCorrectionWindow::SetMeshVisibilityForWindowType()
 {
 	for (EColorCorrectWindowType CCWType : TEnumRange<EColorCorrectWindowType>())
 	{
@@ -160,114 +160,121 @@ void AColorCorrectWindow::SetMeshVisibilityForWindowType()
 		UpdateStageActorTransform();\
 	}\
 
-void AColorCorrectWindow::SetLongitude(double InValue)
+void AColorCorrectionWindow::SetLongitude(double InValue)
 {
 	PositionalParams.Longitude = InValue;
 	NOTIFY_PARAM_SETTER()
 }
 
-double AColorCorrectWindow::GetLongitude() const
+double AColorCorrectionWindow::GetLongitude() const
 {
 	return PositionalParams.Longitude;
 }
 
-void AColorCorrectWindow::SetLatitude(double InValue)
+void AColorCorrectionWindow::SetLatitude(double InValue)
 {
 	PositionalParams.Latitude = InValue;
 	NOTIFY_PARAM_SETTER()
 }
 
-double AColorCorrectWindow::GetLatitude() const
+double AColorCorrectionWindow::GetLatitude() const
 {
 	return PositionalParams.Latitude;
 }
 
-void AColorCorrectWindow::SetDistanceFromCenter(double InValue)
+void AColorCorrectionWindow::SetDistanceFromCenter(double InValue)
 {
 	PositionalParams.DistanceFromCenter = InValue;
 	NOTIFY_PARAM_SETTER()
 }
 
-double AColorCorrectWindow::GetDistanceFromCenter() const
+double AColorCorrectionWindow::GetDistanceFromCenter() const
 {
 	return PositionalParams.DistanceFromCenter;
 }
 
-void AColorCorrectWindow::SetSpin(double InValue)
+void AColorCorrectionWindow::SetSpin(double InValue)
 {
 	PositionalParams.Spin = InValue;
 	NOTIFY_PARAM_SETTER()
 }
 
-double AColorCorrectWindow::GetSpin() const
+double AColorCorrectionWindow::GetSpin() const
 {
 	return PositionalParams.Spin;
 }
 
-void AColorCorrectWindow::SetPitch(double InValue)
+void AColorCorrectionWindow::SetPitch(double InValue)
 {
 	PositionalParams.Pitch = InValue;
 	NOTIFY_PARAM_SETTER()
 }
 
-double AColorCorrectWindow::GetPitch() const
+double AColorCorrectionWindow::GetPitch() const
 {
 	return PositionalParams.Pitch;
 }
 
-void AColorCorrectWindow::SetYaw(double InValue)
+void AColorCorrectionWindow::SetYaw(double InValue)
 {
 	PositionalParams.Yaw = InValue;
 	NOTIFY_PARAM_SETTER()
 }
 
-double AColorCorrectWindow::GetYaw() const
+double AColorCorrectionWindow::GetYaw() const
 {
 	return PositionalParams.Yaw;
 }
 
-void AColorCorrectWindow::SetRadialOffset(double InValue)
+void AColorCorrectionWindow::SetRadialOffset(double InValue)
 {
 	PositionalParams.RadialOffset = InValue;
 	NOTIFY_PARAM_SETTER()
 }
 
-double AColorCorrectWindow::GetRadialOffset() const
+double AColorCorrectionWindow::GetRadialOffset() const
 {
 	return PositionalParams.RadialOffset;
 }
 
-void AColorCorrectWindow::SetScale(const FVector2D& InScale)
+void AColorCorrectionWindow::SetScale(const FVector2D& InScale)
 {
 	PositionalParams.Scale = InScale;
 	NOTIFY_PARAM_SETTER()
 }
 
-FVector2D AColorCorrectWindow::GetScale() const
+FVector2D AColorCorrectionWindow::GetScale() const
 {
 	return PositionalParams.Scale;
 }
 
-void AColorCorrectWindow::SetOrigin(const FTransform& InOrigin)
+void AColorCorrectionWindow::SetOrigin(const FTransform& InOrigin)
 {
 	Origin = InOrigin;
 	NOTIFY_PARAM_SETTER()
 }
 
-FTransform AColorCorrectWindow::GetOrigin() const
+FTransform AColorCorrectionWindow::GetOrigin() const
 {
 	return Origin;
 }
 
-void AColorCorrectWindow::SetPositionalParams(const FDisplayClusterPositionalParams& InParams)
+void AColorCorrectionWindow::SetPositionalParams(const FDisplayClusterPositionalParams& InParams)
 {
 	PositionalParams = InParams;
 	NOTIFY_PARAM_SETTER()
 }
 
-FDisplayClusterPositionalParams AColorCorrectWindow::GetPositionalParams() const
+FDisplayClusterPositionalParams AColorCorrectionWindow::GetPositionalParams() const
 {
 	return PositionalParams;
 }
 
+
+ADEPRECATED_ColorCorrectWindow::ADEPRECATED_ColorCorrectWindow(const FObjectInitializer& ObjectInitializer)
+	: Super(ObjectInitializer)
+{
+
+}
+
 #undef NOTIFY_PARAM_SETTER
diff --git a/Engine/Plugins/Experimental/ColorCorrectRegions/Source/ColorCorrectRegions/Public/ColorCorrectRegion.h b/Engine/Plugins/Experimental/ColorCorrectRegions/Source/ColorCorrectRegions/Public/ColorCorrectRegion.h
--- a/Engine/Plugins/Experimental/ColorCorrectRegions/Source/ColorCorrectRegions/Public/ColorCorrectRegion.h
+++ b/Engine/Plugins/Experimental/ColorCorrectRegions/Source/ColorCorrectRegions/Public/ColorCorrectRegion.h
@@ -57,7 +57,7 @@ public:
  * or during BeginPlay/EndPlay to register itself. 
  * More information in ColorCorrectRegionsSubsytem.h
  */
-UCLASS(Blueprintable, Abstract)
+UCLASS(Blueprintable, NotPlaceable, Abstract)
 class COLORCORRECTREGIONS_API AColorCorrectRegion : public AActor
 {
 	GENERATED_UCLASS_BODY()
@@ -228,3 +228,29 @@ private:
 	float TimeWaited = 0;
 
 };
+
+/** 
+ * NotPlaceable prohibits this actor from appearing in Place Actors menu by default. We add that manually with an icon attached.
+*/
+/**
+ * A placeable Color Correction Regions actor that replaces previous implementation (blueprint). 
+ * Color Correction Regions allow users to adjust color of anything that is within it (or outside, if Invert option is selected). 
+ */
+UCLASS(Blueprintable, NotPlaceable, meta = (DisplayName = "ColorCorrectionRegion"))
+class COLORCORRECTREGIONS_API AColorCorrectionRegion : public AColorCorrectRegion
+{
+	GENERATED_UCLASS_BODY()
+public:
+	UPROPERTY()
+	TArray<TObjectPtr<UStaticMeshComponent>> MeshComponents;
+
+public:
+	/** Swaps meshes for different CCR. */
+	void SetMeshVisibilityForRegionType();
+
+#if WITH_EDITOR
+	/** Called when any of the properties are changed. */
+	virtual void PostEditChangeProperty(FPropertyChangedEvent& PropertyChangedEvent) override;
+
+#endif
+};
diff --git a/Engine/Plugins/Experimental/ColorCorrectRegions/Source/ColorCorrectRegions/Public/ColorCorrectWindow.h b/Engine/Plugins/Experimental/ColorCorrectRegions/Source/ColorCorrectRegions/Public/ColorCorrectWindow.h
--- a/Engine/Plugins/Experimental/ColorCorrectRegions/Source/ColorCorrectRegions/Public/ColorCorrectWindow.h
+++ b/Engine/Plugins/Experimental/ColorCorrectRegions/Source/ColorCorrectRegions/Public/ColorCorrectWindow.h
@@ -26,10 +26,11 @@ enum class EColorCorrectWindowType : uint8
 };
 
 /**
- * 
+ * A Color Correction Window that functions the same way as Color Correction Regions except it modifies anything that is behind it.
+ * Color Correction Windows do not have Priority property and instead are sorted based on the distance from the camera.
  */
-UCLASS(Blueprintable, notplaceable)
-class COLORCORRECTREGIONS_API AColorCorrectWindow : public AColorCorrectRegion, public IDisplayClusterStageActor
+UCLASS(Blueprintable, notplaceable, meta = (DisplayName = "ColorCorrectionWindow"))
+class COLORCORRECTREGIONS_API AColorCorrectionWindow : public AColorCorrectRegion, public IDisplayClusterStageActor
 {
 	GENERATED_UCLASS_BODY()
 public:
@@ -119,3 +120,9 @@ protected:
 	/** Update the transform when a positional setter is called */
 	bool bNotifyOnParamSetter = true;
 };
+
+UCLASS(Deprecated, Blueprintable, notplaceable, meta = (DeprecationMessage = "This is a deprecated version of Color Correct Window. Please re-create Color Correct Window to remove this warning."))
+class COLORCORRECTREGIONS_API ADEPRECATED_ColorCorrectWindow : public AColorCorrectionWindow
+{
+	GENERATED_UCLASS_BODY()
+};
\ No newline at end of file
diff --git a/Engine/Plugins/Experimental/ColorCorrectRegions/Source/ColorCorrectRegionsEditor/Private/ColorCorrectRegionCustomization.cpp b/Engine/Plugins/Experimental/ColorCorrectRegions/Source/ColorCorrectRegionsEditor/Private/ColorCorrectRegionCustomization.cpp
--- a/Engine/Plugins/Experimental/ColorCorrectRegions/Source/ColorCorrectRegionsEditor/Private/ColorCorrectRegionCustomization.cpp
+++ b/Engine/Plugins/Experimental/ColorCorrectRegions/Source/ColorCorrectRegionsEditor/Private/ColorCorrectRegionCustomization.cpp
@@ -14,7 +14,7 @@ void FColorCorrectWindowDetails::CustomizeDetails(IDetailLayoutBuilder& DetailBu
 {
 	for (const TWeakObjectPtr<UObject>& SelectedObject : DetailBuilder.GetSelectedObjects())
 	{
-		if (AColorCorrectWindow* CCW = Cast<AColorCorrectWindow>(SelectedObject.Get()))
+		if (AColorCorrectionWindow* CCW = Cast<AColorCorrectionWindow>(SelectedObject.Get()))
 		{
 			TSharedRef<IPropertyHandle> PriorityProperty = DetailBuilder.GetProperty(GET_MEMBER_NAME_CHECKED(AColorCorrectRegion, Priority));
 			DetailBuilder.HideProperty(PriorityProperty);
diff --git a/Engine/Plugins/Experimental/ColorCorrectRegions/Source/ColorCorrectRegionsEditor/Private/ColorCorrectRegionsEditorModule.cpp b/Engine/Plugins/Experimental/ColorCorrectRegions/Source/ColorCorrectRegionsEditor/Private/ColorCorrectRegionsEditorModule.cpp
--- a/Engine/Plugins/Experimental/ColorCorrectRegions/Source/ColorCorrectRegionsEditor/Private/ColorCorrectRegionsEditorModule.cpp
+++ b/Engine/Plugins/Experimental/ColorCorrectRegions/Source/ColorCorrectRegionsEditor/Private/ColorCorrectRegionsEditorModule.cpp
@@ -27,27 +27,26 @@ void FColorCorrectRegionsEditorModule::OnPlacementModeRefresh(FName CategoryName
 	if (CategoryName == VolumeName || CategoryName == AllClasses)
 	{
 		IPlacementModeModule& PlacementModeModule = IPlacementModeModule::Get();
-		UBlueprint* CCRBlueprint = Cast<UBlueprint>(FSoftObjectPath(TEXT("/ColorCorrectRegions/Blueprints/ColorCorrectRegion.ColorCorrectRegion")).TryLoad());
 		
 		FPlaceableItem* CCRPlaceableItem = new FPlaceableItem(
-			*UActorFactoryBlueprint::StaticClass(),
-			FAssetData(CCRBlueprint, true),
+			*AColorCorrectionRegion::StaticClass(),
+			FAssetData(AColorCorrectionRegion::StaticClass(), true),
 			FName("CCR.PlaceActorThumbnail"),
 			FName("CCR.PlaceActorIcon"),
 			TOptional<FLinearColor>(),
 			TOptional<int32>(),
-			NSLOCTEXT("PlacementMode", "Color Correct Region", "Color Correct Region")
+			NSLOCTEXT("PlacementMode", "Color Correction Region", "Color Correction Region")
 		);
 
 
 		FPlaceableItem* CCWPlaceableItem = new FPlaceableItem(
-			*AColorCorrectWindow::StaticClass(),
-			FAssetData(AColorCorrectWindow::StaticClass()),
+			*AColorCorrectionWindow::StaticClass(),
+			FAssetData(AColorCorrectionWindow::StaticClass()),
 			FName("CCW.PlaceActorThumbnail"),
 			FName("CCR.PlaceActorIcon"),
 			TOptional<FLinearColor>(),
 			TOptional<int32>(), 
-			NSLOCTEXT("PlacementMode", "Color Correct Window", "Color Correct Window"));
+			NSLOCTEXT("PlacementMode", "Color Correction Window", "Color Correction Window"));
 
 		PlacementModeModule.RegisterPlaceableItem(CategoryName, MakeShareable(CCWPlaceableItem));
 		PlacementModeModule.RegisterPlaceableItem(CategoryName, MakeShareable(CCRPlaceableItem));
