Color Correction Regions:
- Added handling for when user selects windows or regions as a per actor CC.
- Added handling for copy pasting and duplication of actors managed by CC.

diff --git a/Engine/Plugins/Experimental/ColorCorrectRegions/Source/ColorCorrectRegions/Private/ColorCorrectRegion.cpp b/Engine/Plugins/Experimental/ColorCorrectRegions/Source/ColorCorrectRegions/Private/ColorCorrectRegion.cpp
--- a/Engine/Plugins/Experimental/ColorCorrectRegions/Source/ColorCorrectRegions/Private/ColorCorrectRegion.cpp
+++ b/Engine/Plugins/Experimental/ColorCorrectRegions/Source/ColorCorrectRegions/Private/ColorCorrectRegion.cpp
@@ -4,6 +4,7 @@
 #include "ColorCorrectWindow.h"
 #include "Components/BillboardComponent.h"
 #include "ColorCorrectRegionsSubsystem.h"
+#include "ColorCorrectRegionsModule.h"
 #include "UObject/ConstructorHelpers.h"
 #include "CoreMinimal.h"
 #include "Engine/GameEngine.h"
@@ -144,7 +145,7 @@ void AColorCorrectRegion::TransferStencilIds()
 {
 	{
 		TArray<uint32> TempStencilIds;
-		for (TSoftObjectPtr<AActor> StencilActor : AffectedActors)
+		for (const TSoftObjectPtr<AActor>& StencilActor : AffectedActors)
 		{
 			if (!StencilActor.IsValid())
 			{
@@ -176,6 +177,26 @@ void AColorCorrectRegion::HandleAffectedActorsPropertyChange()
 		if (ActorListChangeType == EPropertyChangeType::ArrayAdd
 			|| ActorListChangeType == EPropertyChangeType::ValueSet)
 		{
+			// In case user assigns Color Correct Region or Window, we should remove it as it is invalid operation.
+			{
+				TArray<TSoftObjectPtr<AActor>> ActorsToRemove;
+				for (const TSoftObjectPtr<AActor>& StencilActor : AffectedActors)
+				{
+					if (AColorCorrectRegion* CCRCast = Cast<AColorCorrectRegion>(StencilActor.Get()))
+					{
+						ActorsToRemove.Add(StencilActor);
+					}
+				}
+				if (ActorsToRemove.Num() > 0)
+				{
+					UE_LOG(ColorCorrectRegions, Warning, TEXT("Color Correct Region or Window assignment to Per Actor CC is not supported."));
+				}
+				for (const TSoftObjectPtr<AActor>& StencilActor : ActorsToRemove)
+				{
+					AffectedActors.Remove(StencilActor);
+					AffectedActors.FindOrAdd(TSoftObjectPtr<AActor>());
+				}
+			}
 			bEventHandled = true;
 			if (ColorCorrectRegionsSubsystem)
 			{
diff --git a/Engine/Plugins/Experimental/ColorCorrectRegions/Source/ColorCorrectRegions/Private/ColorCorrectRegionSubsystem.cpp b/Engine/Plugins/Experimental/ColorCorrectRegions/Source/ColorCorrectRegions/Private/ColorCorrectRegionSubsystem.cpp
--- a/Engine/Plugins/Experimental/ColorCorrectRegions/Source/ColorCorrectRegions/Private/ColorCorrectRegionSubsystem.cpp
+++ b/Engine/Plugins/Experimental/ColorCorrectRegions/Source/ColorCorrectRegions/Private/ColorCorrectRegionSubsystem.cpp
@@ -41,6 +41,12 @@ void UColorCorrectRegionsSubsystem::Initialize(FSubsystemCollectionBase& Collect
 		GEngine->OnLevelActorDeleted().AddUObject(this, &UColorCorrectRegionsSubsystem::OnActorDeleted);
 		GEngine->OnLevelActorListChanged().AddUObject(this, &UColorCorrectRegionsSubsystem::OnLevelActorListChanged);
 		GEditor->RegisterForUndo(this);
+
+		FEditorDelegates::OnDuplicateActorsBegin.AddUObject(this, &UColorCorrectRegionsSubsystem::OnDuplicateActorsBegin);
+		FEditorDelegates::OnDuplicateActorsEnd.AddUObject(this, &UColorCorrectRegionsSubsystem::OnDuplicateActorsEnd);
+
+		FEditorDelegates::OnEditPasteActorsBegin.AddUObject(this, &UColorCorrectRegionsSubsystem::OnDuplicateActorsBegin);
+		FEditorDelegates::OnEditPasteActorsEnd.AddUObject(this, &UColorCorrectRegionsSubsystem::OnDuplicateActorsEnd);
 	}
 #endif
 	// In some cases (like nDisplay nodes) EndPlay is not guaranteed to be called when level is removed.
@@ -58,6 +64,12 @@ void UColorCorrectRegionsSubsystem::Deinitialize()
 		GEngine->OnLevelActorDeleted().RemoveAll(this);
 		GEngine->OnLevelActorListChanged().RemoveAll(this);
 		GEditor->UnregisterForUndo(this);
+
+		FEditorDelegates::OnDuplicateActorsBegin.RemoveAll(this);
+		FEditorDelegates::OnDuplicateActorsEnd.RemoveAll(this);
+
+		FEditorDelegates::OnEditPasteActorsBegin.RemoveAll(this);
+		FEditorDelegates::OnEditPasteActorsEnd.RemoveAll(this);
 	}
 #endif
 	GetWorld()->OnLevelsChanged().RemoveAll(this);
@@ -113,6 +125,10 @@ void UColorCorrectRegionsSubsystem::OnActorSpawned(AActor* InActor)
 			}
 		}
 	}
+	else if (bDuplicationStarted)
+	{
+		DuplicatedActors.Add(InActor);
+	}
 }
 
 void UColorCorrectRegionsSubsystem::OnActorDeleted(AActor* InActor)
@@ -137,6 +153,32 @@ void UColorCorrectRegionsSubsystem::OnActorDeleted(AActor* InActor)
 	}
 }
 
+void UColorCorrectRegionsSubsystem::OnDuplicateActorsEnd()
+{
+	bDuplicationStarted = false; 
+#if WITH_EDITOR
+	if (GEditor)
+	{
+		GEditor->BeginTransaction(LOCTEXT("PerActorCCActorAssigned", "Per actor CC Actor Assigned"));
+	}
+#endif
+
+	for (AActor* DuplicatedActor : DuplicatedActors)
+	{
+		FColorCorrectRegionsStencilManager::CleanActor(DuplicatedActor);
+	}
+
+	DuplicatedActors.Empty();
+
+#if WITH_EDITOR
+	if (GEditor)
+	{
+		this->Modify();
+		GEditor->EndTransaction();
+	}
+#endif
+}
+
 void UColorCorrectRegionsSubsystem::SortRegionsByPriority()
 {
 	FScopeLock RegionScopeLock(&RegionAccessCriticalSection);
diff --git a/Engine/Plugins/Experimental/ColorCorrectRegions/Source/ColorCorrectRegions/Private/ColorCorrectRegionsStencilManager.cpp b/Engine/Plugins/Experimental/ColorCorrectRegions/Source/ColorCorrectRegions/Private/ColorCorrectRegionsStencilManager.cpp
--- a/Engine/Plugins/Experimental/ColorCorrectRegions/Source/ColorCorrectRegions/Private/ColorCorrectRegionsStencilManager.cpp
+++ b/Engine/Plugins/Experimental/ColorCorrectRegions/Source/ColorCorrectRegions/Private/ColorCorrectRegionsStencilManager.cpp
@@ -379,6 +379,11 @@ void FColorCorrectRegionsStencilManager::AssignStencilIdsToAllActorsForCCR(UWorl
 	}
 }
 
+void FColorCorrectRegionsStencilManager::CleanActor(AActor* Actor)
+{
+	ClearStencilIdFromActor(Actor);
+}
+
 void FColorCorrectRegionsStencilManager::CheckAssignedActorsValidity(AColorCorrectRegion* Region)
 {
 	TRACE_CPUPROFILER_EVENT_SCOPE_TEXT(*FString::Printf(TEXT("CCR.StencilManager.CheckAssignedActorsValidity %s"), *Region->GetName()));
diff --git a/Engine/Plugins/Experimental/ColorCorrectRegions/Source/ColorCorrectRegions/Public/ColorCorrectRegionsStencilManager.h b/Engine/Plugins/Experimental/ColorCorrectRegions/Source/ColorCorrectRegions/Public/ColorCorrectRegionsStencilManager.h
--- a/Engine/Plugins/Experimental/ColorCorrectRegions/Source/ColorCorrectRegions/Public/ColorCorrectRegionsStencilManager.h
+++ b/Engine/Plugins/Experimental/ColorCorrectRegions/Source/ColorCorrectRegions/Public/ColorCorrectRegionsStencilManager.h
@@ -15,6 +15,7 @@ public:
 		static void OnCCRRemoved(UWorld* CurrentWorld, AColorCorrectRegion* Region);
 		static void ClearInvalidActorsForSelectedRegion(AColorCorrectRegion* Region);
 		static void AssignStencilIdsToAllActorsForCCR(UWorld* CurrentWorld, AColorCorrectRegion* Region, bool bIgnoreUserNotifications, bool bSoftAssign);
+		static void CleanActor(AActor* Actor);
 
 		/** Runs over all actors assigned to the selected Region and makes sure that Stencil Ids match the ones assigned to it. */
 		static void CheckAssignedActorsValidity(AColorCorrectRegion* Region);
diff --git a/Engine/Plugins/Experimental/ColorCorrectRegions/Source/ColorCorrectRegions/Public/ColorCorrectRegionsSubsystem.h b/Engine/Plugins/Experimental/ColorCorrectRegions/Source/ColorCorrectRegions/Public/ColorCorrectRegionsSubsystem.h
--- a/Engine/Plugins/Experimental/ColorCorrectRegions/Source/ColorCorrectRegions/Public/ColorCorrectRegionsSubsystem.h
+++ b/Engine/Plugins/Experimental/ColorCorrectRegions/Source/ColorCorrectRegions/Public/ColorCorrectRegionsSubsystem.h
@@ -71,7 +71,13 @@ public:
 
 	/** Called when level is added or removed. */
 	void OnLevelsChanged() { RefreshRegions(); };
-	
+
+	/** Called when duplication process is started in the level. */
+	void OnDuplicateActorsBegin() { bDuplicationStarted = true; };
+
+	/** Called when duplication process is ended in the level. */
+	void OnDuplicateActorsEnd();
+
 #if WITH_EDITOR
 	/** A callback for when the level is loaded. */
 	void OnLevelActorListChanged() { RefreshRegions(); };
@@ -114,6 +120,10 @@ private:
 
 	FCriticalSection RegionAccessCriticalSection;
 
+	/** This is to handle actor duplication for Per Actor CC. */
+	bool bDuplicationStarted = false;
+	TArray<AActor*> DuplicatedActors;
+
 public:
 	friend class FColorCorrectRegionsSceneViewExtension;
 };
\ No newline at end of file
