Color Correction Regions: Changed CCR and CCW Icons as well as the plugin icon.
- removed explicit naming of CCW and CCR
- Icons are now all svgs
- Changed icons for Outliner and ICVFX Lightcard editor

diff --git a/Engine/Plugins/Experimental/ColorCorrectRegions/Source/ColorCorrectRegions/Private/ColorCorrectRegion.cpp b/Engine/Plugins/Experimental/ColorCorrectRegions/Source/ColorCorrectRegions/Private/ColorCorrectRegion.cpp
--- a/Engine/Plugins/Experimental/ColorCorrectRegions/Source/ColorCorrectRegions/Private/ColorCorrectRegion.cpp
+++ b/Engine/Plugins/Experimental/ColorCorrectRegions/Source/ColorCorrectRegions/Private/ColorCorrectRegion.cpp
@@ -359,4 +359,8 @@ void AColorCorrectionRegion::PostEditChangeProperty(struct FPropertyChangedEvent
 		SetMeshVisibilityForRegionType();
 	}
 }
+FName AColorCorrectionRegion::GetCustomIconName() const
+{
+	return TEXT("CCR.OutlinerThumbnail");
+}
 #endif
diff --git a/Engine/Plugins/Experimental/ColorCorrectRegions/Source/ColorCorrectRegions/Private/ColorCorrectWindow.cpp b/Engine/Plugins/Experimental/ColorCorrectRegions/Source/ColorCorrectRegions/Private/ColorCorrectWindow.cpp
--- a/Engine/Plugins/Experimental/ColorCorrectRegions/Source/ColorCorrectRegions/Private/ColorCorrectWindow.cpp
+++ b/Engine/Plugins/Experimental/ColorCorrectRegions/Source/ColorCorrectRegions/Private/ColorCorrectWindow.cpp
@@ -90,6 +90,12 @@ void AColorCorrectionWindow::PostEditMove(bool bFinished)
 	UpdatePositionalParamsFromTransform();
 	bNotifyOnParamSetter = true;
 }
+
+FName AColorCorrectionWindow::GetCustomIconName() const
+{
+	return TEXT("CCW.OutlinerThumbnail");
+}
+
 #endif //WITH_EDITOR
 
 
diff --git a/Engine/Plugins/Experimental/ColorCorrectRegions/Source/ColorCorrectRegions/Public/ColorCorrectRegion.h b/Engine/Plugins/Experimental/ColorCorrectRegions/Source/ColorCorrectRegions/Public/ColorCorrectRegion.h
--- a/Engine/Plugins/Experimental/ColorCorrectRegions/Source/ColorCorrectRegions/Public/ColorCorrectRegion.h
+++ b/Engine/Plugins/Experimental/ColorCorrectRegions/Source/ColorCorrectRegions/Public/ColorCorrectRegion.h
@@ -236,7 +236,7 @@ private:
  * A placeable Color Correction Regions actor that replaces previous implementation (blueprint). 
  * Color Correction Regions allow users to adjust color of anything that is within it (or outside, if Invert option is selected). 
  */
-UCLASS(Blueprintable, NotPlaceable, meta = (DisplayName = "ColorCorrectionRegion"))
+UCLASS(Blueprintable, NotPlaceable)
 class COLORCORRECTREGIONS_API AColorCorrectionRegion : public AColorCorrectRegion
 {
 	GENERATED_UCLASS_BODY()
@@ -251,6 +251,6 @@ public:
 #if WITH_EDITOR
 	/** Called when any of the properties are changed. */
 	virtual void PostEditChangeProperty(FPropertyChangedEvent& PropertyChangedEvent) override;
-
+	virtual FName GetCustomIconName() const override;
 #endif
 };
diff --git a/Engine/Plugins/Experimental/ColorCorrectRegions/Source/ColorCorrectRegions/Public/ColorCorrectWindow.h b/Engine/Plugins/Experimental/ColorCorrectRegions/Source/ColorCorrectRegions/Public/ColorCorrectWindow.h
--- a/Engine/Plugins/Experimental/ColorCorrectRegions/Source/ColorCorrectRegions/Public/ColorCorrectWindow.h
+++ b/Engine/Plugins/Experimental/ColorCorrectRegions/Source/ColorCorrectRegions/Public/ColorCorrectWindow.h
@@ -29,7 +29,7 @@ enum class EColorCorrectWindowType : uint8
  * A Color Correction Window that functions the same way as Color Correction Regions except it modifies anything that is behind it.
  * Color Correction Windows do not have Priority property and instead are sorted based on the distance from the camera.
  */
-UCLASS(Blueprintable, notplaceable, meta = (DisplayName = "ColorCorrectionWindow"))
+UCLASS(Blueprintable, notplaceable)
 class COLORCORRECTREGIONS_API AColorCorrectionWindow : public AColorCorrectRegion, public IDisplayClusterStageActor
 {
 	GENERATED_UCLASS_BODY()
@@ -46,6 +46,7 @@ public:
 	/** Called when any of the properties are changed. */
 	virtual void PostEditChangeProperty(FPropertyChangedEvent& PropertyChangedEvent) override;
 	virtual void PostEditMove(bool bFinished) override;
+	virtual FName GetCustomIconName() const override;
 #endif
 
 private:
diff --git a/Engine/Plugins/Experimental/ColorCorrectRegions/Source/ColorCorrectRegionsEditor/Private/ColorCorrectRegionsEditorModule.cpp b/Engine/Plugins/Experimental/ColorCorrectRegions/Source/ColorCorrectRegionsEditor/Private/ColorCorrectRegionsEditorModule.cpp
--- a/Engine/Plugins/Experimental/ColorCorrectRegions/Source/ColorCorrectRegionsEditor/Private/ColorCorrectRegionsEditorModule.cpp
+++ b/Engine/Plugins/Experimental/ColorCorrectRegions/Source/ColorCorrectRegionsEditor/Private/ColorCorrectRegionsEditorModule.cpp
@@ -32,18 +32,17 @@ void FColorCorrectRegionsEditorModule::OnPlacementModeRefresh(FName CategoryName
 			*AColorCorrectionRegion::StaticClass(),
 			FAssetData(AColorCorrectionRegion::StaticClass(), true),
 			FName("CCR.PlaceActorThumbnail"),
-			FName("CCR.PlaceActorIcon"),
+			FName("CCR.OutlinerThumbnail"),
 			TOptional<FLinearColor>(),
 			TOptional<int32>(),
 			NSLOCTEXT("PlacementMode", "Color Correction Region", "Color Correction Region")
 		);
 
-
 		FPlaceableItem* CCWPlaceableItem = new FPlaceableItem(
 			*AColorCorrectionWindow::StaticClass(),
 			FAssetData(AColorCorrectionWindow::StaticClass()),
 			FName("CCW.PlaceActorThumbnail"),
-			FName("CCR.PlaceActorIcon"),
+			FName("CCW.OutlinerThumbnail"),
 			TOptional<FLinearColor>(),
 			TOptional<int32>(), 
 			NSLOCTEXT("PlacementMode", "Color Correction Window", "Color Correction Window"));
diff --git a/Engine/Plugins/Experimental/ColorCorrectRegions/Source/ColorCorrectRegionsEditor/Private/ColorCorrectRegionsStyle.cpp b/Engine/Plugins/Experimental/ColorCorrectRegions/Source/ColorCorrectRegionsEditor/Private/ColorCorrectRegionsStyle.cpp
--- a/Engine/Plugins/Experimental/ColorCorrectRegions/Source/ColorCorrectRegionsEditor/Private/ColorCorrectRegionsStyle.cpp
+++ b/Engine/Plugins/Experimental/ColorCorrectRegions/Source/ColorCorrectRegionsEditor/Private/ColorCorrectRegionsStyle.cpp
@@ -17,11 +17,14 @@ void FColorCorrectRegionsStyle::Initialize()
 
 	CCRStyle = MakeShared<FSlateStyleSet>(FName("ColorCorrectRegionsStyle"));
 
-	FString CCR_IconPath = IPluginManager::Get().FindPlugin(TEXT("ColorCorrectRegions"))->GetBaseDir() + TEXT("/Resources/PlaceActorPreview.png");
-	FString CCW_IconPath = IPluginManager::Get().FindPlugin(TEXT("ColorCorrectRegions"))->GetBaseDir() + TEXT("/Resources/Icon40_CCW.png");
-	CCRStyle->Set("CCR.PlaceActorThumbnail", new FSlateImageBrush(CCR_IconPath, FVector2D(40.0f, 40.0f)));
-	CCRStyle->Set("CCW.PlaceActorThumbnail", new FSlateImageBrush(CCW_IconPath, FVector2D(40.0f, 40.0f)));
-	CCRStyle->Set("CCR.PlaceActorIcon", new FSlateImageBrush(CCR_IconPath, FVector2D(16.0f, 16.0f)));
+	FString CCR_IconPath = IPluginManager::Get().FindPlugin(TEXT("ColorCorrectRegions"))->GetBaseDir() + TEXT("/Resources/Template_CCR_64.svg");
+	FString CCW_IconPath = IPluginManager::Get().FindPlugin(TEXT("ColorCorrectRegions"))->GetBaseDir() + TEXT("/Resources/Template_CCW_64.svg");
+
+	CCRStyle->Set("CCR.PlaceActorThumbnail", new FSlateVectorImageBrush(CCR_IconPath, FVector2D(40.0f, 40.0f)));
+	CCRStyle->Set("CCW.PlaceActorThumbnail", new FSlateVectorImageBrush(CCW_IconPath, FVector2D(40.0f, 40.0f)));
+
+	CCRStyle->Set("CCR.OutlinerThumbnail", new FSlateVectorImageBrush(CCR_IconPath, FVector2D(20.0f, 20.0f)));
+	CCRStyle->Set("CCW.OutlinerThumbnail", new FSlateVectorImageBrush(CCW_IconPath, FVector2D(20.0f, 20.0f)));
 
 	FSlateStyleRegistry::RegisterSlateStyle(*CCRStyle.Get());
 }
