Color Correct Regions: Fixed an MU issue where two Editors are connected and one editor assigns an actor or CCR.

diff --git a/Engine/Plugins/Experimental/ColorCorrectRegions/Source/ColorCorrectRegions/Private/ColorCorrectRegionsStencilManager.cpp b/Engine/Plugins/Experimental/ColorCorrectRegions/Source/ColorCorrectRegions/Private/ColorCorrectRegionsStencilManager.cpp
--- a/Engine/Plugins/Experimental/ColorCorrectRegions/Source/ColorCorrectRegions/Private/ColorCorrectRegionsStencilManager.cpp
+++ b/Engine/Plugins/Experimental/ColorCorrectRegions/Source/ColorCorrectRegions/Private/ColorCorrectRegionsStencilManager.cpp
@@ -81,7 +81,7 @@ namespace
 			// If the stencil id is already assigned and not by another CCR then we need to notify the user.
 			if ((PrimitiveComponent->bRenderCustomDepth && PrimitiveComponent->CustomDepthStencilValue < StencilIdMin))
 			{
-				EAppReturnType::Type Answer = EAppReturnType::No;
+				EAppReturnType::Type Answer = EAppReturnType::Yes;
 				if (!bIgnoreUserNotification)
 				{
 					auto Text = FString::Printf(TEXT("The actor \'%s\' already has Custom Depth Stencil Id assigned that is out of range.\
@@ -399,6 +399,10 @@ void FColorCorrectRegionsStencilManager::CheckAssignedActorsValidity(AColorCorre
 				TArray<UPrimitiveComponent*> PrimitiveComponents;
 
 				Actor->GetComponents<UPrimitiveComponent>(PrimitiveComponents);
+				if (AColorCorrectRegion* CCRCast = Cast<AColorCorrectRegion>(Actor.Get()))
+				{
+					continue;
+				}
 				for (UPrimitiveComponent* PrimitiveComponent : PrimitiveComponents)
 				{
 					if (PrimitiveComponent->CustomDepthStencilValue != StencilDataPair.Value->AssignedStencil || !PrimitiveComponent->bRenderCustomDepth)
