Color Correct Regions: Added class icon for ICVFX

diff --git a/Engine/Plugins/Experimental/ColorCorrectRegions/Source/ColorCorrectRegionsEditor/Private/ColorCorrectRegionsStyle.cpp b/Engine/Plugins/Experimental/ColorCorrectRegions/Source/ColorCorrectRegionsEditor/Private/ColorCorrectRegionsStyle.cpp
--- a/Engine/Plugins/Experimental/ColorCorrectRegions/Source/ColorCorrectRegionsEditor/Private/ColorCorrectRegionsStyle.cpp
+++ b/Engine/Plugins/Experimental/ColorCorrectRegions/Source/ColorCorrectRegionsEditor/Private/ColorCorrectRegionsStyle.cpp
@@ -25,6 +25,10 @@ void FColorCorrectRegionsStyle::Initialize()
 
 	CCRStyle->Set("CCR.OutlinerThumbnail", new FSlateVectorImageBrush(CCR_IconPath, FVector2D(20.0f, 20.0f)));
 	CCRStyle->Set("CCW.OutlinerThumbnail", new FSlateVectorImageBrush(CCW_IconPath, FVector2D(20.0f, 20.0f)));
+	CCRStyle->Set("ClassIcon.ColorCorrectionWindow", new FSlateVectorImageBrush(CCW_IconPath, FVector2D(20.0f, 20.0f)));
+	CCRStyle->Set("ClassIcon.ColorCorrectionRegion", new FSlateVectorImageBrush(CCR_IconPath, FVector2D(20.0f, 20.0f)));
+	CCRStyle->Set("ClassThumbnail.ColorCorrectionWindow", new FSlateVectorImageBrush(CCW_IconPath, FVector2D(20.0f, 20.0f)));
+	CCRStyle->Set("ClassThumbnail.ColorCorrectionRegion", new FSlateVectorImageBrush(CCR_IconPath, FVector2D(20.0f, 20.0f)));
 
 	FSlateStyleRegistry::RegisterSlateStyle(*CCRStyle.Get());
 }
