Color Correct Regions: Fixed the issue with deprecated CCRs misbehaving in conjunction with Sequencer.

diff --git a/Engine/Plugins/Experimental/ColorCorrectRegions/Config/BaseColorCorrectRegions.ini b/Engine/Plugins/Experimental/ColorCorrectRegions/Config/BaseColorCorrectRegions.ini
new file mode 100644
--- /dev/null
+++ b/Engine/Plugins/Experimental/ColorCorrectRegions/Config/BaseColorCorrectRegions.ini
@@ -0,0 +1,3 @@
+[CoreRedirects]
+
++ClassRedirects=(OldName="ColorCorrectRegion_C",NewName="/Script/ColorCorrectRegions.ColorCorrectionRegion",OverrideClassName="/Script/CoreUObject.Class")
\ No newline at end of file
