Color Correct Regions: Removed warning about missing unused CCR Blueprint.
- removed unecessary CCR shadow casting.

diff --git a/Engine/Plugins/Experimental/ColorCorrectRegions/Config/BaseColorCorrectRegions.ini b/Engine/Plugins/Experimental/ColorCorrectRegions/Config/BaseColorCorrectRegions.ini
--- a/Engine/Plugins/Experimental/ColorCorrectRegions/Config/BaseColorCorrectRegions.ini
+++ b/Engine/Plugins/Experimental/ColorCorrectRegions/Config/BaseColorCorrectRegions.ini
@@ -1,3 +1,4 @@
 [CoreRedirects]
 
-+ClassRedirects=(OldName="ColorCorrectRegion_C",NewName="/Script/ColorCorrectRegions.ColorCorrectionRegion",OverrideClassName="/Script/CoreUObject.Class")
\ No newline at end of file
++ClassRedirects=(OldName="ColorCorrectRegion_C",NewName="/Script/ColorCorrectRegions.ColorCorrectionRegion",OverrideClassName="/Script/CoreUObject.Class")
++PackageRedirects=(OldName="/ColorCorrectRegions/Blueprints/ColorCorrectRegion",Removed=True)
\ No newline at end of file
diff --git a/Engine/Plugins/Experimental/ColorCorrectRegions/Source/ColorCorrectRegions/Private/ColorCorrectRegion.cpp b/Engine/Plugins/Experimental/ColorCorrectRegions/Source/ColorCorrectRegions/Private/ColorCorrectRegion.cpp
--- a/Engine/Plugins/Experimental/ColorCorrectRegions/Source/ColorCorrectRegions/Private/ColorCorrectRegion.cpp
+++ b/Engine/Plugins/Experimental/ColorCorrectRegions/Source/ColorCorrectRegions/Private/ColorCorrectRegion.cpp
@@ -325,6 +325,7 @@ AColorCorrectionRegion::AColorCorrectionRegion(const FObjectInitializer& ObjectI
 		MeshComponent->SetMaterial(0, Material);
 		MeshComponent->SetCollisionEnabled(ECollisionEnabled::QueryOnly);
 		MeshComponent->SetCollisionProfileName(TEXT("OverlapAll"));
+		MeshComponent->CastShadow = false;
 	}
 
 	SetMeshVisibilityForRegionType();
diff --git a/Engine/Plugins/Experimental/ColorCorrectRegions/Source/ColorCorrectRegions/Private/ColorCorrectWindow.cpp b/Engine/Plugins/Experimental/ColorCorrectRegions/Source/ColorCorrectRegions/Private/ColorCorrectWindow.cpp
--- a/Engine/Plugins/Experimental/ColorCorrectRegions/Source/ColorCorrectRegions/Private/ColorCorrectWindow.cpp
+++ b/Engine/Plugins/Experimental/ColorCorrectRegions/Source/ColorCorrectRegions/Private/ColorCorrectWindow.cpp
@@ -42,6 +42,7 @@ AColorCorrectionWindow::AColorCorrectionWindow(const FObjectInitializer& ObjectI
 		MeshComponent->SetMaterial(0, Material);
 		MeshComponent->SetCollisionEnabled(ECollisionEnabled::QueryOnly);
 		MeshComponent->SetCollisionProfileName(TEXT("OverlapAll"));
+		MeshComponent->CastShadow = false;
 	}
 	SetMeshVisibilityForWindowType();
 
