Fixed #include paths.

diff --git a/Engine/Plugins/Enterprise/DatasmithC4DImporter/Source/DatasmithC4DTranslator/Private/DatasmithC4DMelangeSDKOverload.cpp b/Engine/Plugins/Enterprise/DatasmithC4DImporter/Source/DatasmithC4DTranslator/Private/DatasmithC4DMelangeSDKOverload.cpp
--- a/Engine/Plugins/Enterprise/DatasmithC4DImporter/Source/DatasmithC4DTranslator/Private/DatasmithC4DMelangeSDKOverload.cpp
+++ b/Engine/Plugins/Enterprise/DatasmithC4DImporter/Source/DatasmithC4DTranslator/Private/DatasmithC4DMelangeSDKOverload.cpp
@@ -2,7 +2,7 @@
 
 #ifdef _MELANGE_SDK_
 
-#include <Runtime/Core/Public/HAL/UnrealMemory.h>
+#include "HAL/UnrealMemory.h"
 
 #include "DatasmithC4DMelangeSDKEnterGuard.h"
 #define DONT_INCLUDE_MEMORY_OVERLOADS
@@ -72,4 +72,4 @@ void GetWriterInfo(cineware::Int32 &id, cineware::String &appname)
 	appname = "Datasmith C4D Importer";
 }
 
-#endif //_MELANGE_SDK_
\ No newline at end of file
+#endif //_MELANGE_SDK_
diff --git a/Engine/Plugins/Enterprise/DatasmithContent/Source/DatasmithContent/DatasmithContent.Build.cs b/Engine/Plugins/Enterprise/DatasmithContent/Source/DatasmithContent/DatasmithContent.Build.cs
--- a/Engine/Plugins/Enterprise/DatasmithContent/Source/DatasmithContent/DatasmithContent.Build.cs
+++ b/Engine/Plugins/Enterprise/DatasmithContent/Source/DatasmithContent/DatasmithContent.Build.cs
@@ -30,6 +30,12 @@ namespace UnrealBuildTool.Rules
 				}
 			);
 
+			PrivateIncludePathModuleNames.AddRange(
+				new string[] {
+					"MainFrame",
+				}
+				);
+
 			if (Target.Type == TargetType.Editor)
 			{
 				PrivateDependencyModuleNames.AddRange(
diff --git a/Engine/Plugins/Enterprise/DatasmithContent/Source/DatasmithContent/Private/DatasmithImportOptions.cpp b/Engine/Plugins/Enterprise/DatasmithContent/Source/DatasmithContent/Private/DatasmithImportOptions.cpp
--- a/Engine/Plugins/Enterprise/DatasmithContent/Source/DatasmithContent/Private/DatasmithImportOptions.cpp
+++ b/Engine/Plugins/Enterprise/DatasmithContent/Source/DatasmithContent/Private/DatasmithImportOptions.cpp
@@ -5,7 +5,7 @@
 #include "UObject/UnrealType.h"
 
 #include "HAL/FileManager.h"
-#include "Editor/MainFrame/Public/Interfaces/IMainFrameModule.h"
+#include "Interfaces/IMainFrameModule.h"
 
 #include "Dom/JsonValue.h"
 #include "Dom/JsonObject.h"
diff --git a/Engine/Plugins/Experimental/ColorCorrectRegions/Source/ColorCorrectRegions/Private/ColorCorrectRegion.cpp b/Engine/Plugins/Experimental/ColorCorrectRegions/Source/ColorCorrectRegions/Private/ColorCorrectRegion.cpp
--- a/Engine/Plugins/Experimental/ColorCorrectRegions/Source/ColorCorrectRegions/Private/ColorCorrectRegion.cpp
+++ b/Engine/Plugins/Experimental/ColorCorrectRegions/Source/ColorCorrectRegions/Private/ColorCorrectRegion.cpp
@@ -1,18 +1,18 @@
 // Copyright Epic Games, Inc. All Rights Reserved.
 
 #include "ColorCorrectRegion.h"
+#include "ColorCorrectRegionsModule.h"
+#include "ColorCorrectRegionsSubsystem.h"
 #include "ColorCorrectWindow.h"
 #include "Components/BillboardComponent.h"
-#include "ColorCorrectRegionsSubsystem.h"
-#include "ColorCorrectRegionsModule.h"
-#include "UObject/ConstructorHelpers.h"
+#include "Components/MeshComponent.h"
+#include "Components/StaticMeshComponent.h"
 #include "CoreMinimal.h"
 #include "Engine/GameEngine.h"
-#include "Engine/Classes/Components/MeshComponent.h"
-#include "Materials/Material.h"
 #include "Engine/StaticMesh.h"
-#include "Components/StaticMeshComponent.h"
 #include "Engine/Texture2D.h"
+#include "Materials/Material.h"
+#include "UObject/ConstructorHelpers.h"
 
 ENUM_RANGE_BY_COUNT(EColorCorrectRegionsType, EColorCorrectRegionsType::MAX)
 
diff --git a/Engine/Plugins/Experimental/ColorCorrectRegions/Source/ColorCorrectRegions/Private/ColorCorrectRegionsModule.cpp b/Engine/Plugins/Experimental/ColorCorrectRegions/Source/ColorCorrectRegions/Private/ColorCorrectRegionsModule.cpp
--- a/Engine/Plugins/Experimental/ColorCorrectRegions/Source/ColorCorrectRegions/Private/ColorCorrectRegionsModule.cpp
+++ b/Engine/Plugins/Experimental/ColorCorrectRegions/Source/ColorCorrectRegions/Private/ColorCorrectRegionsModule.cpp
@@ -1,9 +1,9 @@
 // Copyright Epic Games, Inc. All Rights Reserved.
 
 #include "ColorCorrectRegionsModule.h"
-#include "Runtime/Core/Public/Misc/Paths.h"
 #include "Interfaces/IPluginManager.h"
-#include "Runtime/RenderCore/Public/ShaderCore.h"
+#include "Misc/Paths.h"
+#include "ShaderCore.h"
 
 
 #define LOCTEXT_NAMESPACE "FColorCorrectRegionsModule"
diff --git a/Engine/Plugins/Experimental/ColorCorrectRegions/Source/ColorCorrectRegions/Public/ColorCorrectRegion.h b/Engine/Plugins/Experimental/ColorCorrectRegions/Source/ColorCorrectRegions/Public/ColorCorrectRegion.h
--- a/Engine/Plugins/Experimental/ColorCorrectRegions/Source/ColorCorrectRegions/Public/ColorCorrectRegion.h
+++ b/Engine/Plugins/Experimental/ColorCorrectRegions/Source/ColorCorrectRegions/Public/ColorCorrectRegion.h
@@ -2,12 +2,12 @@
 
 #pragma once
 
+#include "Components/MeshComponent.h"
 #include "CoreMinimal.h"
-#include "UObject/ObjectMacros.h"
-#include "GameFramework/Actor.h"
-#include "Engine/Classes/Components/MeshComponent.h"
 #include "Engine/Scene.h"
+#include "GameFramework/Actor.h"
 #include "Templates/SharedPointer.h"
+#include "UObject/ObjectMacros.h"
 #include "ColorCorrectRegion.generated.h"
 
 
diff --git a/Engine/Plugins/Experimental/ColorCorrectRegions/Source/ColorCorrectRegions/Public/ColorCorrectRegionsPostProcessMaterial.h b/Engine/Plugins/Experimental/ColorCorrectRegions/Source/ColorCorrectRegions/Public/ColorCorrectRegionsPostProcessMaterial.h
--- a/Engine/Plugins/Experimental/ColorCorrectRegions/Source/ColorCorrectRegions/Public/ColorCorrectRegionsPostProcessMaterial.h
+++ b/Engine/Plugins/Experimental/ColorCorrectRegions/Source/ColorCorrectRegions/Public/ColorCorrectRegionsPostProcessMaterial.h
@@ -1,20 +1,20 @@
 // Copyright Epic Games, Inc. All Rights Reserved.
 
 #pragma once
+#include "ColorCorrectRegion.h"
+#include "ColorCorrectWindow.h"
 #include "CoreMinimal.h"
-#include "UObject/ConstructorHelpers.h"
-#include "RHI.h"
+#include "MaterialShader.h"
+#include "Materials/Material.h"
 #include "Materials/MaterialInstanceDynamic.h"
 #include "Materials/MaterialInterface.h"
-#include "Materials/Material.h"
-#include "Runtime/RHI/Public/RHIStaticStates.h"
-#include "Runtime/RenderCore/Public/ShaderParameterUtils.h"
-#include "Runtime/RenderCore/Public/RenderResource.h"
-#include "Runtime/Renderer/Public/MaterialShader.h"
-#include "Runtime/RenderCore/Public/RenderGraphResources.h"
-#include "Runtime/RenderCore/Public/RenderGraphResources.h"
-#include "ColorCorrectRegion.h"
-#include "ColorCorrectWindow.h"
+#include "RHI.h"
+#include "RHIStaticStates.h"
+#include "RenderGraphResources.h"
+#include "RenderGraphResources.h"
+#include "RenderResource.h"
+#include "ShaderParameterUtils.h"
+#include "UObject/ConstructorHelpers.h"
 
 // FScreenPassTextureViewportParameters and FScreenPassTextureInput
 #include "Runtime/Renderer/Private/ScreenPass.h"
@@ -266,4 +266,4 @@ class FCCRStencilMergerPS : public FGlobalShader
 		SHADER_PARAMETER_STRUCT(FScreenPassTextureViewportParameters, PostProcessOutput)
 		RENDER_TARGET_BINDING_SLOTS()
 	END_SHADER_PARAMETER_STRUCT()
-};
\ No newline at end of file
+};
diff --git a/Engine/Plugins/Experimental/ColorCorrectRegions/Source/ColorCorrectRegions/Public/ColorCorrectWindow.h b/Engine/Plugins/Experimental/ColorCorrectRegions/Source/ColorCorrectRegions/Public/ColorCorrectWindow.h
--- a/Engine/Plugins/Experimental/ColorCorrectRegions/Source/ColorCorrectRegions/Public/ColorCorrectWindow.h
+++ b/Engine/Plugins/Experimental/ColorCorrectRegions/Source/ColorCorrectRegions/Public/ColorCorrectWindow.h
@@ -3,11 +3,11 @@
 #pragma once
 
 #include "ColorCorrectRegion.h"
+#include "Components/MeshComponent.h"
 #include "CoreMinimal.h"
-#include "UObject/ObjectMacros.h"
-#include "GameFramework/Actor.h"
-#include "Engine/Classes/Components/MeshComponent.h"
 #include "Engine/Scene.h"
+#include "GameFramework/Actor.h"
+#include "UObject/ObjectMacros.h"
 
 #include "StageActor/IDisplayClusterStageActor.h"
 
@@ -128,4 +128,4 @@ UCLASS(Deprecated, Blueprintable, notplaceable, meta = (DeprecationMessage = "Th
 class COLORCORRECTREGIONS_API ADEPRECATED_ColorCorrectWindow : public AColorCorrectionWindow
 {
 	GENERATED_UCLASS_BODY()
-};
\ No newline at end of file
+};
diff --git a/Engine/Plugins/Experimental/OptiXDenoise/Source/OptiXDenoise/Private/Plugin.cpp b/Engine/Plugins/Experimental/OptiXDenoise/Source/OptiXDenoise/Private/Plugin.cpp
--- a/Engine/Plugins/Experimental/OptiXDenoise/Source/OptiXDenoise/Private/Plugin.cpp
+++ b/Engine/Plugins/Experimental/OptiXDenoise/Source/OptiXDenoise/Private/Plugin.cpp
@@ -1,7 +1,7 @@
 // Copyright Epic Games, Inc. All Rights Reserved.
 
 #include "OptiXDenoiser.h"
-#include "Renderer/Public/PathTracingDenoiser.h"
+#include "PathTracingDenoiser.h"
 #include <mutex>
 
 #define LOCTEXT_NAMESPACE "FOptiXDenoiseModule"
@@ -247,4 +247,4 @@ void FOptiXDenoiseModule::ShutdownModule()
 
 #undef LOCTEXT_NAMESPACE
 
-IMPLEMENT_MODULE(FOptiXDenoiseModule, OptiXDenoise)
\ No newline at end of file
+IMPLEMENT_MODULE(FOptiXDenoiseModule, OptiXDenoise)
diff --git a/Engine/Plugins/Experimental/WebSocketNetworking/Source/WebSocketNetworking/Private/WebsocketConnection.cpp b/Engine/Plugins/Experimental/WebSocketNetworking/Source/WebSocketNetworking/Private/WebsocketConnection.cpp
--- a/Engine/Plugins/Experimental/WebSocketNetworking/Source/WebSocketNetworking/Private/WebsocketConnection.cpp
+++ b/Engine/Plugins/Experimental/WebSocketNetworking/Source/WebSocketNetworking/Private/WebsocketConnection.cpp
@@ -7,11 +7,11 @@
 #include "Net/NetworkProfiler.h"
 #include "Net/DataChannel.h"
 
+#include "PacketHandler.h"
+#include "PacketHandlers/StatelessConnectHandlerComponent.h"
+#include "WebSocket.h"
 #include "WebSocketConnection.h"
 #include "WebSocketNetDriver.h"
-#include "WebSocket.h"
-#include "Runtime/PacketHandlers/PacketHandler/Public/PacketHandler.h"
-#include "PacketHandlers/StatelessConnectHandlerComponent.h"
 
 /*-----------------------------------------------------------------------------
 Declarations.
diff --git a/Engine/Plugins/Media/HLMedia/Source/HLMedia/Private/HLMediaPlayer.cpp b/Engine/Plugins/Media/HLMedia/Source/HLMedia/Private/HLMediaPlayer.cpp
--- a/Engine/Plugins/Media/HLMedia/Source/HLMedia/Private/HLMediaPlayer.cpp
+++ b/Engine/Plugins/Media/HLMedia/Source/HLMedia/Private/HLMediaPlayer.cpp
@@ -3,11 +3,11 @@
 #include "HLMediaPlayer.h"
 
 #include "Async/Async.h"
+#include "GenericPlatform/GenericPlatformProcess.h"
 #include "IMediaEventSink.h"
 #include "IMediaOptions.h"
-#include "MediaUtils/Public/MediaPlayerOptions.h"
-#include "MediaUtils/Public/MediaSamples.h"
-#include "GenericPlatform/GenericPlatformProcess.h"
+#include "MediaPlayerOptions.h"
+#include "MediaSamples.h"
 #include "Misc/MessageDialog.h"
 #include "RenderingThread.h"
 
diff --git a/Engine/Plugins/Media/HLMedia/Source/HLMedia/Private/HLMediaPlayerTracks.cpp b/Engine/Plugins/Media/HLMedia/Source/HLMedia/Private/HLMediaPlayerTracks.cpp
--- a/Engine/Plugins/Media/HLMedia/Source/HLMedia/Private/HLMediaPlayerTracks.cpp
+++ b/Engine/Plugins/Media/HLMedia/Source/HLMedia/Private/HLMediaPlayerTracks.cpp
@@ -2,7 +2,7 @@
 
 #include "HLMediaPlayerTracks.h"
 
-#include "MediaUtils/Public/MediaSamples.h"
+#include "MediaSamples.h"
 
 using namespace HLMediaLibrary;
 
diff --git a/Engine/Plugins/Media/HLMedia/Source/HLMedia/Private/HLMediaPlayerTracks.h b/Engine/Plugins/Media/HLMedia/Source/HLMedia/Private/HLMediaPlayerTracks.h
--- a/Engine/Plugins/Media/HLMedia/Source/HLMedia/Private/HLMediaPlayerTracks.h
+++ b/Engine/Plugins/Media/HLMedia/Source/HLMedia/Private/HLMediaPlayerTracks.h
@@ -6,15 +6,15 @@
 
 #include "IMediaTracks.h"
 
-#include "CoreTypes.h"
 #include "Containers/Array.h"
 #include "Containers/UnrealString.h"
+#include "CoreTypes.h"
 #include "HAL/CriticalSection.h"
 #include "Internationalization/Text.h"
 #include "Math/IntPoint.h"
 #include "Math/Range.h"
+#include "MediaPlayerOptions.h"
 #include "Templates/Function.h"
-#include "MediaUtils/Public/MediaPlayerOptions.h"
 
 #include "HLMediaTextureSample.h"
 
diff --git a/Engine/Plugins/Media/HLMedia/Source/HLMedia/Private/HLMediaTextureSample.h b/Engine/Plugins/Media/HLMedia/Source/HLMedia/Private/HLMediaTextureSample.h
--- a/Engine/Plugins/Media/HLMedia/Source/HLMedia/Private/HLMediaTextureSample.h
+++ b/Engine/Plugins/Media/HLMedia/Source/HLMedia/Private/HLMediaTextureSample.h
@@ -5,7 +5,7 @@
 #include "HLMediaPrivate.h"
 
 #include "IMediaTextureSample.h"
-#include "MediaUtils/Public/MediaObjectPool.h"
+#include "MediaObjectPool.h"
 
 #include "ID3D11DynamicRHI.h"
 #include "ID3D12DynamicRHI.h"
diff --git a/Engine/Plugins/Media/ImgMedia/Source/ExrReaderGpu/Private/ExrReaderGpuModule.cpp b/Engine/Plugins/Media/ImgMedia/Source/ExrReaderGpu/Private/ExrReaderGpuModule.cpp
--- a/Engine/Plugins/Media/ImgMedia/Source/ExrReaderGpu/Private/ExrReaderGpuModule.cpp
+++ b/Engine/Plugins/Media/ImgMedia/Source/ExrReaderGpu/Private/ExrReaderGpuModule.cpp
@@ -5,7 +5,7 @@
 
 #include "HAL/Platform.h"
 #include "Interfaces/IPluginManager.h"
-#include "Runtime/Core/Public/Misc/Paths.h"
+#include "Misc/Paths.h"
 #include "ShaderCore.h"
 
 DEFINE_LOG_CATEGORY(LogExrReaderGpu);
diff --git a/Engine/Plugins/Media/PixelCapture/Source/PixelCapture/Private/PixelCaptureUtils.h b/Engine/Plugins/Media/PixelCapture/Source/PixelCapture/Private/PixelCaptureUtils.h
--- a/Engine/Plugins/Media/PixelCapture/Source/PixelCapture/Private/PixelCaptureUtils.h
+++ b/Engine/Plugins/Media/PixelCapture/Source/PixelCapture/Private/PixelCaptureUtils.h
@@ -2,11 +2,11 @@
 
 #pragma once
 
+#include "MediaShaders.h"
 #include "Modules/ModuleManager.h"
 #include "RHI.h"
-#include "Runtime/Renderer/Private/ScreenPass.h"
+#include "ScreenPass.h"
 #include "ScreenRendering.h"
-#include "MediaShaders.h"
 
 /*
     * Copy from one texture to another.
@@ -137,4 +137,4 @@ inline void CopyTextureRDG(FRHICommandListImmediate& RHICmdList, FTextureRHIRef
 		AddDrawScreenPass(GraphBuilder, RDG_EVENT_NAME("PixelCapturerSwizzle"), ViewInfo, OutputViewport, InputViewport, VertexShader, PixelShader, Parameters);
 	}
 	GraphBuilder.Execute();
-}
\ No newline at end of file
+}
diff --git a/Engine/Plugins/Media/PixelStreaming/Source/PixelStreaming/Private/UtilsRender.h b/Engine/Plugins/Media/PixelStreaming/Source/PixelStreaming/Private/UtilsRender.h
--- a/Engine/Plugins/Media/PixelStreaming/Source/PixelStreaming/Private/UtilsRender.h
+++ b/Engine/Plugins/Media/PixelStreaming/Source/PixelStreaming/Private/UtilsRender.h
@@ -2,11 +2,11 @@
 
 #pragma once
 
+#include "Modules/ModuleManager.h"
 #include "RHIGPUReadback.h"
 #include "RenderGraphBuilder.h"
 #include "RenderGraphUtils.h"
-#include "Runtime/Renderer/Private/ScreenPass.h"
-#include "Modules/ModuleManager.h"
+#include "ScreenPass.h"
 #include "ScreenRendering.h"
 #include "Utils.h"
 
diff --git a/Engine/Plugins/Online/OnlineSubsystemOculus/Source/Private/OnlineVoiceOculus.cpp b/Engine/Plugins/Online/OnlineSubsystemOculus/Source/Private/OnlineVoiceOculus.cpp
--- a/Engine/Plugins/Online/OnlineSubsystemOculus/Source/Private/OnlineVoiceOculus.cpp
+++ b/Engine/Plugins/Online/OnlineSubsystemOculus/Source/Private/OnlineVoiceOculus.cpp
@@ -6,8 +6,8 @@
 #include "Components/AudioComponent.h"
 #include "Misc/ConfigCacheIni.h"
 
-#include "Runtime/Engine/Classes/Sound/SoundWaveProcedural.h"
 #include "OnlineSubsystemUtils.h"
+#include "Sound/SoundWaveProcedural.h"
 
 #include "Misc/ConfigCacheIni.h"
 #include "OnlineSubsystemOculusPackage.h"
diff --git a/Engine/Plugins/Runtime/AR/Google/GoogleARCore/Source/GoogleARCoreBase/Private/GoogleARCoreOpenGLContext.h b/Engine/Plugins/Runtime/AR/Google/GoogleARCore/Source/GoogleARCoreBase/Private/GoogleARCoreOpenGLContext.h
--- a/Engine/Plugins/Runtime/AR/Google/GoogleARCore/Source/GoogleARCoreBase/Private/GoogleARCoreOpenGLContext.h
+++ b/Engine/Plugins/Runtime/AR/Google/GoogleARCore/Source/GoogleARCoreBase/Private/GoogleARCoreOpenGLContext.h
@@ -2,8 +2,8 @@
 
 #pragma once
 
-#include "Runtime/Core/Public/Misc/Build.h"
 #include "HAL/Platform.h"
+#include "Misc/Build.h"
 #include "Templates/SharedPointer.h"
 
 #if PLATFORM_ANDROID
diff --git a/Engine/Plugins/Runtime/OSC/Source/OSC/Private/OSCServer.cpp b/Engine/Plugins/Runtime/OSC/Source/OSC/Private/OSCServer.cpp
--- a/Engine/Plugins/Runtime/OSC/Source/OSC/Private/OSCServer.cpp
+++ b/Engine/Plugins/Runtime/OSC/Source/OSC/Private/OSCServer.cpp
@@ -1,10 +1,10 @@
 // Copyright Epic Games, Inc. All Rights Reserved.
 #include "OSCServer.h"
 
+#include "Async/TaskGraphInterfaces.h"
 #include "Common/UdpSocketBuilder.h"
 #include "Common/UdpSocketReceiver.h"
 #include "Interfaces/IPv4/IPv4Address.h"
-#include "Runtime/Core/Public/Async/TaskGraphInterfaces.h"
 #include "Sockets.h"
 
 #include "OSCStream.h"
@@ -251,4 +251,4 @@ void UOSCServer::PumpPacketQueue(const TSet<uint32>* AllowlistedClients)
 			}
 		}
 	}
-}
\ No newline at end of file
+}
diff --git a/Engine/Plugins/Runtime/Oculus/OculusAudio/Source/OculusAudio/Private/OculusAudioGeometryComponent.cpp b/Engine/Plugins/Runtime/Oculus/OculusAudio/Source/OculusAudio/Private/OculusAudioGeometryComponent.cpp
--- a/Engine/Plugins/Runtime/Oculus/OculusAudio/Source/OculusAudio/Private/OculusAudioGeometryComponent.cpp
+++ b/Engine/Plugins/Runtime/Oculus/OculusAudio/Source/OculusAudio/Private/OculusAudioGeometryComponent.cpp
@@ -1,15 +1,15 @@
 // Copyright Epic Games, Inc. All Rights Reserved.
 #include "OculusAudioGeometryComponent.h"
-#include "Components/StaticMeshComponent.h"
-#include "OculusAudioMaterialComponent.h"
-#include "OculusAudioContextManager.h"
-#include "OculusAudio.h"
 #include "AudioDevice.h"
-#include "IOculusAudioPlugin.h"
-#include "Engine/StaticMeshActor.h"
+#include "Components/StaticMeshComponent.h"
 #include "Engine/StaticMesh.h"
-#include "Runtime/Core/Public/Serialization/CustomVersion.h"
+#include "Engine/StaticMeshActor.h"
+#include "IOculusAudioPlugin.h"
 #include "Materials/MaterialInstanceDynamic.h"
+#include "OculusAudio.h"
+#include "OculusAudioContextManager.h"
+#include "OculusAudioMaterialComponent.h"
+#include "Serialization/CustomVersion.h"
 #include "StaticMeshResources.h"
 
 
diff --git a/Engine/Plugins/Runtime/Oculus/OculusVR/Source/OculusHMD/Private/OculusHMD.cpp b/Engine/Plugins/Runtime/Oculus/OculusVR/Source/OculusHMD/Private/OculusHMD.cpp
--- a/Engine/Plugins/Runtime/Oculus/OculusVR/Source/OculusHMD/Private/OculusHMD.cpp
+++ b/Engine/Plugins/Runtime/Oculus/OculusVR/Source/OculusHMD/Private/OculusHMD.cpp
@@ -42,7 +42,7 @@
 #include "IOculusMRModule.h"
 
 #if WITH_EDITOR
-#include "Editor/UnrealEd/Classes/Editor/EditorEngine.h"
+#include "Editor/EditorEngine.h"
 #endif
 
 #if !UE_BUILD_SHIPPING
diff --git a/Engine/Plugins/Runtime/Oculus/OculusVR/Source/OculusHMD/Private/OculusHMD_Layer.cpp b/Engine/Plugins/Runtime/Oculus/OculusVR/Source/OculusHMD/Private/OculusHMD_Layer.cpp
--- a/Engine/Plugins/Runtime/Oculus/OculusVR/Source/OculusHMD/Private/OculusHMD_Layer.cpp
+++ b/Engine/Plugins/Runtime/Oculus/OculusVR/Source/OculusHMD/Private/OculusHMD_Layer.cpp
@@ -7,18 +7,18 @@
 //#include "ScreenRendering.h"
 //#include "ScenePrivate.h"
 //#include "PostProcess/SceneFilterRendering.h"
-#include "PostProcess/SceneRenderTargets.h"
-#include "HeadMountedDisplayTypes.h" // for LogHMD
-#include "OculusHMD.h"
-#include "XRThreadUtils.h"
+#include "Engine/GameEngine.h"
 #include "Engine/Texture2D.h"
-#include "UObject/ConstructorHelpers.h"
-#include "Materials/MaterialInterface.h"
+#include "HeadMountedDisplayTypes.h" // for LogHMD
 #include "Materials/MaterialInstanceDynamic.h"
-#include "Engine/GameEngine.h"
-#include "Engine/Public/SceneUtils.h"
-#include "OculusHMDPrivate.h"
+#include "Materials/MaterialInterface.h"
+#include "OculusHMD.h"
 #include "OculusHMDModule.h"
+#include "OculusHMDPrivate.h"
+#include "PostProcess/SceneRenderTargets.h"
+#include "SceneUtils.h"
+#include "UObject/ConstructorHelpers.h"
+#include "XRThreadUtils.h"
 
 
 namespace OculusHMD
diff --git a/Engine/Plugins/Runtime/OpenXR/Source/OpenXRHMD/Private/OpenXRHMD.cpp b/Engine/Plugins/Runtime/OpenXR/Source/OpenXRHMD/Private/OpenXRHMD.cpp
--- a/Engine/Plugins/Runtime/OpenXR/Source/OpenXRHMD/Private/OpenXRHMD.cpp
+++ b/Engine/Plugins/Runtime/OpenXR/Source/OpenXRHMD/Private/OpenXRHMD.cpp
@@ -37,7 +37,7 @@
 #include "DefaultStereoLayers.h"
 
 #if WITH_EDITOR
-#include "Editor/UnrealEd/Classes/Editor/EditorEngine.h"
+#include "Editor/EditorEngine.h"
 #endif
 
 #define OPENXR_PAUSED_IDLE_FPS 10
diff --git a/Engine/Plugins/Runtime/WebMMoviePlayer/Source/WebMMoviePlayer/Private/Audio/Windows/WindowsWebMAudioBackendDefines.h b/Engine/Plugins/Runtime/WebMMoviePlayer/Source/WebMMoviePlayer/Private/Audio/Windows/WindowsWebMAudioBackendDefines.h
--- a/Engine/Plugins/Runtime/WebMMoviePlayer/Source/WebMMoviePlayer/Private/Audio/Windows/WindowsWebMAudioBackendDefines.h
+++ b/Engine/Plugins/Runtime/WebMMoviePlayer/Source/WebMMoviePlayer/Private/Audio/Windows/WindowsWebMAudioBackendDefines.h
@@ -2,4 +2,4 @@
 
 #pragma once
 
-#include "../Null/WebMAudioBackendDefines.h"
+#include "Audio/Null/WebMAudioBackendDefines.h"
diff --git a/Engine/Plugins/Runtime/WindowsMoviePlayer/Source/WindowsMoviePlayer/Private/WindowsMovieStreamer.cpp b/Engine/Plugins/Runtime/WindowsMoviePlayer/Source/WindowsMoviePlayer/Private/WindowsMovieStreamer.cpp
--- a/Engine/Plugins/Runtime/WindowsMoviePlayer/Source/WindowsMoviePlayer/Private/WindowsMovieStreamer.cpp
+++ b/Engine/Plugins/Runtime/WindowsMoviePlayer/Source/WindowsMoviePlayer/Private/WindowsMovieStreamer.cpp
@@ -2,11 +2,11 @@
 
 #include "WindowsMovieStreamer.h"
 
-#include "Rendering/RenderingCommon.h"
-#include "Slate/SlateTextures.h"
+#include "IHeadMountedDisplayModule.h"
 #include "MoviePlayer.h"
 #include "RenderUtils.h"
-#include "Runtime/HeadMountedDisplay/Public/IHeadMountedDisplayModule.h"
+#include "Rendering/RenderingCommon.h"
+#include "Slate/SlateTextures.h"
 
 #include "MediaShaders.h"
 #include "PipelineStateCache.h"
diff --git a/Engine/Plugins/Runtime/WindowsMoviePlayer/Source/WindowsMoviePlayer/WindowsMoviePlayer.Build.cs b/Engine/Plugins/Runtime/WindowsMoviePlayer/Source/WindowsMoviePlayer/WindowsMoviePlayer.Build.cs
--- a/Engine/Plugins/Runtime/WindowsMoviePlayer/Source/WindowsMoviePlayer/WindowsMoviePlayer.Build.cs
+++ b/Engine/Plugins/Runtime/WindowsMoviePlayer/Source/WindowsMoviePlayer/WindowsMoviePlayer.Build.cs
@@ -13,7 +13,8 @@ namespace UnrealBuildTool.Rules
 
 			PrivateIncludePathModuleNames.AddRange(
 				new string[] {
-					"Media"
+					"HeadMountedDisplay",
+					"Media",
 				});
 
 			PrivateIncludePaths.AddRange(
diff --git a/Engine/Plugins/Runtime/nDisplay/Source/DisplayCluster/Private/Misc/DisplayClusterAppExit.cpp b/Engine/Plugins/Runtime/nDisplay/Source/DisplayCluster/Private/Misc/DisplayClusterAppExit.cpp
--- a/Engine/Plugins/Runtime/nDisplay/Source/DisplayCluster/Private/Misc/DisplayClusterAppExit.cpp
+++ b/Engine/Plugins/Runtime/nDisplay/Source/DisplayCluster/Private/Misc/DisplayClusterAppExit.cpp
@@ -7,8 +7,8 @@
 #include "Async/Async.h"
 
 #if WITH_EDITOR
-#include "Editor/UnrealEd/Public/UnrealEdGlobals.h"
 #include "Editor/UnrealEdEngine.h"
+#include "UnrealEdGlobals.h"
 #endif
 
 
diff --git a/Engine/Plugins/Runtime/nDisplay/Source/DisplayClusterShaders/Private/Shaders/DisplayClusterShadersPreprocess_UVLightCards.cpp b/Engine/Plugins/Runtime/nDisplay/Source/DisplayClusterShaders/Private/Shaders/DisplayClusterShadersPreprocess_UVLightCards.cpp
--- a/Engine/Plugins/Runtime/nDisplay/Source/DisplayClusterShaders/Private/Shaders/DisplayClusterShadersPreprocess_UVLightCards.cpp
+++ b/Engine/Plugins/Runtime/nDisplay/Source/DisplayClusterShaders/Private/Shaders/DisplayClusterShadersPreprocess_UVLightCards.cpp
@@ -12,7 +12,7 @@
 #include "PrimitiveSceneProxy.h"
 #include "RenderGraphBuilder.h"
 #include "RenderGraphResources.h"
-#include "Renderer/Private/ScenePrivate.h"
+#include "ScenePrivate.h"
 #include "SceneRenderTargetParameters.h"
 #include "MeshPassProcessor.inl"
 #include "UnrealClient.h"
@@ -256,4 +256,4 @@ bool FDisplayClusterShadersPreprocess_UVLightCards::RenderPreprocess_UVLightCard
 	GraphBuilder.Execute();
 
 	return true;
-}
\ No newline at end of file
+}
diff --git a/Engine/Plugins/VirtualProduction/RemoteControl/Source/RemoteControl/Private/RemoteControlBinding.cpp b/Engine/Plugins/VirtualProduction/RemoteControl/Source/RemoteControl/Private/RemoteControlBinding.cpp
--- a/Engine/Plugins/VirtualProduction/RemoteControl/Source/RemoteControl/Private/RemoteControlBinding.cpp
+++ b/Engine/Plugins/VirtualProduction/RemoteControl/Source/RemoteControl/Private/RemoteControlBinding.cpp
@@ -17,7 +17,6 @@
 #if WITH_EDITOR
 #include "Components/ActorComponent.h"
 #include "Editor.h"
-#include "Editor/UnrealEd/Public/Editor.h"
 #include "Misc/App.h"
 #endif
 
@@ -436,4 +435,4 @@ void URemoteControlLevelDependantBinding::SetBoundObject(const TSoftObjectPtr<UL
 	Name = MoveTemp(ObjectName);
 }
 
-#undef LOCTEXT_NAMESPACE 
\ No newline at end of file
+#undef LOCTEXT_NAMESPACE 
diff --git a/Engine/Plugins/VirtualProduction/RemoteControl/Source/RemoteControl/Private/RemoteControlPresetRebindingManager.cpp b/Engine/Plugins/VirtualProduction/RemoteControl/Source/RemoteControl/Private/RemoteControlPresetRebindingManager.cpp
--- a/Engine/Plugins/VirtualProduction/RemoteControl/Source/RemoteControl/Private/RemoteControlPresetRebindingManager.cpp
+++ b/Engine/Plugins/VirtualProduction/RemoteControl/Source/RemoteControl/Private/RemoteControlPresetRebindingManager.cpp
@@ -4,10 +4,10 @@
 
 #include "Algo/AllOf.h"
 #include "Algo/Transform.h"
-#include "Engine/BlueprintGeneratedClass.h"
 #include "Components/ActorComponent.h"
+#include "Components/ActorComponent.h"
+#include "Engine/BlueprintGeneratedClass.h"
 #include "Engine/Brush.h"
-#include "Engine/Classes/Components/ActorComponent.h"
 #include "Engine/Engine.h"
 #include "EngineUtils.h"
 #include "Misc/Char.h"
diff --git a/Engine/Plugins/VirtualProduction/RemoteControl/Source/RemoteControl/Public/RemoteControlBinding.h b/Engine/Plugins/VirtualProduction/RemoteControl/Source/RemoteControl/Public/RemoteControlBinding.h
--- a/Engine/Plugins/VirtualProduction/RemoteControl/Source/RemoteControl/Public/RemoteControlBinding.h
+++ b/Engine/Plugins/VirtualProduction/RemoteControl/Source/RemoteControl/Public/RemoteControlBinding.h
@@ -4,12 +4,12 @@
 
 #include "UObject/Object.h"
 
+#include "Components/ActorComponent.h"
 #include "GameFramework/Actor.h"
-#include "Engine/Classes/Components/ActorComponent.h"
 #include "UObject/ObjectMacros.h"
 #include "UObject/SoftObjectPtr.h"
-#include "UObject/WeakObjectPtrTemplates.h"
 #include "UObject/UnrealType.h"
+#include "UObject/WeakObjectPtrTemplates.h"
 #include "RemoteControlBinding.generated.h"
 
 namespace UE::RemoteControlBinding
diff --git a/Engine/Plugins/VirtualProduction/TextureShare/Source/TextureShare/Private/Game/WorldSubsystem/TextureShareWorldSubsystemObject.cpp b/Engine/Plugins/VirtualProduction/TextureShare/Source/TextureShare/Private/Game/WorldSubsystem/TextureShareWorldSubsystemObject.cpp
--- a/Engine/Plugins/VirtualProduction/TextureShare/Source/TextureShare/Private/Game/WorldSubsystem/TextureShareWorldSubsystemObject.cpp
+++ b/Engine/Plugins/VirtualProduction/TextureShare/Source/TextureShare/Private/Game/WorldSubsystem/TextureShareWorldSubsystemObject.cpp
@@ -1,8 +1,8 @@
 // Copyright Epic Games, Inc. All Rights Reserved.
 
 #include "Game/WorldSubsystem/TextureShareWorldSubsystemObject.h"
+#include "Game/ViewExtension/TextureShareSceneViewExtension.h"
 #include "Game/WorldSubsystem/TextureShareWorldSubsystemObjectProxy.h"
-#include "Game/ViewExtension//TextureShareSceneViewExtension.h"
 
 #include "Containers/TextureShareContainers.h"
 
diff --git a/Engine/Source/Runtime/Android/AndroidLocalNotification/Private/AndroidLocalNotification.cpp b/Engine/Source/Runtime/Android/AndroidLocalNotification/Private/AndroidLocalNotification.cpp
--- a/Engine/Source/Runtime/Android/AndroidLocalNotification/Private/AndroidLocalNotification.cpp
+++ b/Engine/Source/Runtime/Android/AndroidLocalNotification/Private/AndroidLocalNotification.cpp
@@ -8,9 +8,9 @@
 	Includes
  ------------------------------------------------------------------------------------*/
 
-#include "../Public/AndroidLocalNotification.h"
-#include "Stats/Stats.h"
+#include "AndroidLocalNotification.h"
 #include "Async/TaskGraphInterfaces.h"
+#include "Stats/Stats.h"
 
 DEFINE_LOG_CATEGORY(LogAndroidLocalNotification);
 
diff --git a/Engine/Source/Runtime/Android/AndroidLocalNotification/Public/AndroidLocalNotification.h b/Engine/Source/Runtime/Android/AndroidLocalNotification/Public/AndroidLocalNotification.h
--- a/Engine/Source/Runtime/Android/AndroidLocalNotification/Public/AndroidLocalNotification.h
+++ b/Engine/Source/Runtime/Android/AndroidLocalNotification/Public/AndroidLocalNotification.h
@@ -13,7 +13,7 @@
 	#include <android_native_app_glue.h>
 #endif
 
-#include "Runtime/Core/Public/Modules/ModuleManager.h"
+#include "Modules/ModuleManager.h"
 
 DECLARE_LOG_CATEGORY_EXTERN(LogAndroidLocalNotification, Log, All);
 
diff --git a/Engine/Source/Runtime/Core/Private/FramePro/FrameProProfiler.cpp b/Engine/Source/Runtime/Core/Private/FramePro/FrameProProfiler.cpp
--- a/Engine/Source/Runtime/Core/Private/FramePro/FrameProProfiler.cpp
+++ b/Engine/Source/Runtime/Core/Private/FramePro/FrameProProfiler.cpp
@@ -5,11 +5,11 @@
 #if FRAMEPRO_ENABLED
 
 #include "CoreTypes.h"
-#include "HAL/ThreadSingleton.h"
-#include "HAL/IConsoleManager.h"
 #include "HAL/FileManager.h"
+#include "HAL/IConsoleManager.h"
 #include "HAL/MallocAnsi.h"
 #include "HAL/ThreadManager.h"
+#include "HAL/ThreadSingleton.h"
 #include "Misc/App.h"
 #include "Misc/EngineVersion.h"
 #include "Misc/Paths.h"
diff --git a/Engine/Source/Runtime/Engine/Private/SparseVolumeTexture/SparseVolumeTexture.cpp b/Engine/Source/Runtime/Engine/Private/SparseVolumeTexture/SparseVolumeTexture.cpp
--- a/Engine/Source/Runtime/Engine/Private/SparseVolumeTexture/SparseVolumeTexture.cpp
+++ b/Engine/Source/Runtime/Engine/Private/SparseVolumeTexture/SparseVolumeTexture.cpp
@@ -11,7 +11,7 @@
 #include "Shader/ShaderTypes.h"
 
 #if WITH_EDITORONLY_DATA
-#include "DerivedDataCache/Public/DerivedDataCacheInterface.h"
+#include "DerivedDataCacheInterface.h"
 #endif
 
 #include "Serialization/LargeMemoryReader.h"
diff --git a/Engine/Source/Runtime/Experimental/Iris/Core/Private/Iris/ReplicationSystem/ObjectReferenceCache.cpp b/Engine/Source/Runtime/Experimental/Iris/Core/Private/Iris/ReplicationSystem/ObjectReferenceCache.cpp
--- a/Engine/Source/Runtime/Experimental/Iris/Core/Private/Iris/ReplicationSystem/ObjectReferenceCache.cpp
+++ b/Engine/Source/Runtime/Experimental/Iris/Core/Private/Iris/ReplicationSystem/ObjectReferenceCache.cpp
@@ -1,25 +1,25 @@
 // Copyright Epic Games, Inc. All Rights Reserved.
 
 #include "ObjectReferenceCache.h"
-#include "Iris/IrisConstants.h"
 #include "Iris/Core/IrisLog.h"
 #include "Iris/Core/IrisMemoryTracker.h"
+#include "Iris/IrisConstants.h"
 #include "Iris/ReplicationSystem/ObjectReplicationBridge.h"
 #include "Iris/ReplicationSystem/ReplicationSystem.h"
 #include "Iris/ReplicationSystem/ReplicationSystemInternal.h"
 #include "Iris/ReplicationSystem/StringTokenStore.h"
-#include "Iris/Serialization/NetBitStreamWriter.h"
+#include "Iris/Serialization/InternalNetSerializationContext.h"
 #include "Iris/Serialization/NetBitStreamReader.h"
+#include "Iris/Serialization/NetBitStreamUtil.h"
+#include "Iris/Serialization/NetBitStreamWriter.h"
+#include "Iris/Serialization/NetExportContext.h"
+#include "Iris/Serialization/NetSerializationContext.h"
 #include "Iris/Serialization/ObjectNetSerializer.h"
-#include "UObject/Package.h"
 #include "Misc/CoreMiscDefines.h"
+#include "Misc/StringBuilder.h"
 #include "Net/Core/Trace/NetDebugName.h"
 #include "Net/Core/Trace/NetTrace.h"
-#include "Core/Public/Misc/StringBuilder.h"
-#include "Iris/Serialization/NetSerializationContext.h"
-#include "Iris/Serialization/InternalNetSerializationContext.h"
-#include "Iris/Serialization/NetExportContext.h"
-#include "Iris/Serialization/NetBitStreamUtil.h"
+#include "UObject/Package.h"
 
 DEFINE_LOG_CATEGORY_STATIC(LogIrisReferences, Log, All);
 
diff --git a/Engine/Source/Runtime/Launch/Private/Android/AndroidJNI.cpp b/Engine/Source/Runtime/Launch/Private/Android/AndroidJNI.cpp
--- a/Engine/Source/Runtime/Launch/Private/Android/AndroidJNI.cpp
+++ b/Engine/Source/Runtime/Launch/Private/Android/AndroidJNI.cpp
@@ -3,20 +3,20 @@
 #include "Android/AndroidJNI.h"
 
 #if USE_ANDROID_JNI
-#include "HAL/ExceptionHandling.h"
-#include "Android/AndroidPlatformCrashContext.h"
-#include "Runtime/Core/Public/Misc/DateTime.h"
-#include "HAL/PlatformStackWalk.h"
 #include "Android/AndroidApplication.h"
 #include "Android/AndroidInputInterface.h"
-#include "Widgets/Input/IVirtualKeyboardEntry.h"
-#include "UnrealEngine.h"
-#include "Misc/ConfigCacheIni.h"
-#include "Misc/FeedbackContext.h"
+#include "Android/AndroidPlatformCrashContext.h"
+#include "Android/AndroidStats.h"
+#include "Async/TaskGraphInterfaces.h"
+#include "HAL/ExceptionHandling.h"
+#include "HAL/PlatformStackWalk.h"
 #include "Math/Vector.h"
+#include "Misc/ConfigCacheIni.h"
+#include "Misc/DateTime.h"
 #include "Misc/EmbeddedCommunication.h"
-#include "Async/TaskGraphInterfaces.h"
-#include "Android/AndroidStats.h"
+#include "Misc/FeedbackContext.h"
+#include "UnrealEngine.h"
+#include "Widgets/Input/IVirtualKeyboardEntry.h"
 
 THIRD_PARTY_INCLUDES_START
 #include <android/asset_manager.h>
diff --git a/Engine/Source/Runtime/Renderer/Private/SparseVolumeTexture/SparseVolumeTextureViewerSceneProxy.cpp b/Engine/Source/Runtime/Renderer/Private/SparseVolumeTexture/SparseVolumeTextureViewerSceneProxy.cpp
--- a/Engine/Source/Runtime/Renderer/Private/SparseVolumeTexture/SparseVolumeTextureViewerSceneProxy.cpp
+++ b/Engine/Source/Runtime/Renderer/Private/SparseVolumeTexture/SparseVolumeTextureViewerSceneProxy.cpp
@@ -4,9 +4,9 @@
 #include "SparseVolumeTexture/SparseVolumeTextureViewerComponent.h"
 #include "SparseVolumeTexture/SparseVolumeTexture.h"
 
-#include "SceneView.h"
-#include "Renderer/Private/ScenePrivate.h"
 #include "RenderGraphUtils.h"
+#include "ScenePrivate.h"
+#include "SceneView.h"
 
 #define LOCTEXT_NAMESPACE "FSparseVolumeTextureViewerSceneProxy"
 
