Color Correct Windows: Added redirect for old windows.

diff --git a/Engine/Plugins/Experimental/ColorCorrectRegions/Config/BaseColorCorrectRegions.ini b/Engine/Plugins/Experimental/ColorCorrectRegions/Config/BaseColorCorrectRegions.ini
--- a/Engine/Plugins/Experimental/ColorCorrectRegions/Config/BaseColorCorrectRegions.ini
+++ b/Engine/Plugins/Experimental/ColorCorrectRegions/Config/BaseColorCorrectRegions.ini
@@ -1,4 +1,6 @@
 [CoreRedirects]
 
 +ClassRedirects=(OldName="ColorCorrectRegion_C",NewName="/Script/ColorCorrectRegions.ColorCorrectionRegion",OverrideClassName="/Script/CoreUObject.Class")
-+PackageRedirects=(OldName="/ColorCorrectRegions/Blueprints/ColorCorrectRegion",Removed=True)
\ No newline at end of file
++PackageRedirects=(OldName="/ColorCorrectRegions/Blueprints/ColorCorrectRegion",Removed=True)
++ClassRedirects=(OldName="/Script/ColorCorrectRegions.ColorCorrectWindow", NewName="/Script/ColorCorrectRegions.ColorCorrectionWindow")
+
