Color Correction Regions: From now on all properties used on render thread will be accessed via a state that is a copy of requried properties.

diff --git a/Engine/Plugins/Experimental/ColorCorrectRegions/Source/ColorCorrectRegions/Private/ColorCorrectRegion.cpp b/Engine/Plugins/Experimental/ColorCorrectRegions/Source/ColorCorrectRegions/Private/ColorCorrectRegion.cpp
--- a/Engine/Plugins/Experimental/ColorCorrectRegions/Source/ColorCorrectRegions/Private/ColorCorrectRegion.cpp
+++ b/Engine/Plugins/Experimental/ColorCorrectRegions/Source/ColorCorrectRegions/Private/ColorCorrectRegion.cpp
@@ -17,7 +17,7 @@
 ENUM_RANGE_BY_COUNT(EColorCorrectRegionsType, EColorCorrectRegionsType::MAX)
 
 
-AColorCorrectRegion::AColorCorrectRegion(const FObjectInitializer& ObjectInitializer) 
+AColorCorrectRegion::AColorCorrectRegion(const FObjectInitializer& ObjectInitializer)
 	: Super(ObjectInitializer)
 	, Type(EColorCorrectRegionsType::Sphere)
 	, Priority(0)
@@ -33,6 +33,7 @@ AColorCorrectRegion::AColorCorrectRegion(const FObjectInitializer& ObjectInitial
 	, bEnablePerActorCC(false)
 	, PerActorColorCorrection(EColorCorrectRegionStencilType::IncludeStencil)
 	, ColorCorrectRegionsSubsystem(nullptr)
+	, ColorCorrectRenderProxy(MakeShared<FColorCorrectRenderProxy>())
 {
 	PrimaryActorTick.bCanEverTick = true;
 
@@ -95,7 +96,8 @@ void AColorCorrectRegion::TickActor(float DeltaTime, ELevelTick TickType, FActor
 	TRACE_CPUPROFILER_EVENT_SCOPE_TEXT(*FString::Printf(TEXT("CCR.TickActor %s"), *GetName()));
 
 	HandleAffectedActorsPropertyChange();
-	TransferStencilIds();
+	TransferState();
+
 
 	// Check to make sure that no ids have been changed externally.
 	{
@@ -117,23 +119,6 @@ void AColorCorrectRegion::TickActor(float DeltaTime, ELevelTick TickType, FActor
 		}
 
 	}
-
-	FTransform CurrentFrameTransform = GetTransform();
-	if (!PreviousFrameTransform.Equals(CurrentFrameTransform))
-	{
-		PreviousFrameTransform = CurrentFrameTransform;
-		GetActorBounds(true, BoxOrigin, BoxExtent);
-	}
-	
-	// Display Cluster uses HiddenPrimitives to hide Primitive components from view. 
-	// Store component id to be used on render thread.
-	if (const UStaticMeshComponent* FirstMeshComponent = FindComponentByClass<UStaticMeshComponent>())
-	{
-		if (!(FirstPrimitiveId == FirstMeshComponent->ComponentId))
-		{
-			FirstPrimitiveId = FirstMeshComponent->ComponentId;
-		}
-	}
 }
 
 void AColorCorrectRegion::Cleanup()
@@ -141,10 +126,57 @@ void AColorCorrectRegion::Cleanup()
 	ColorCorrectRegionsSubsystem = nullptr;
 }
 
-void AColorCorrectRegion::TransferStencilIds()
+void AColorCorrectRegion::TransferState()
 {
+	TRACE_CPUPROFILER_EVENT_SCOPE_TEXT(*FString::Printf(TEXT("CCR.TransferState %s"), *GetName()));
+
+	FColorCorrectRenderProxyPtr TempCCRStateRenderThread = MakeShared<FColorCorrectRenderProxy>();
+	if (AColorCorrectionWindow* CCWindow = Cast<AColorCorrectionWindow>(this))
+	{
+		TempCCRStateRenderThread->WindowType = CCWindow->WindowType;
+	}
+	else
 	{
-		TArray<uint32> TempStencilIds;
+		TempCCRStateRenderThread->Type = Type;
+	}
+
+	TempCCRStateRenderThread->bIsActiveThisFrame = IsValid(this)
+		&& Enabled
+#if WITH_EDITOR
+		&& !IsHiddenEd()
+#endif 
+		&& !IsHidden();
+
+	TempCCRStateRenderThread->World = GetWorld();
+	TempCCRStateRenderThread->Priority = Priority;
+	TempCCRStateRenderThread->Intensity = Intensity;
+
+	// Inner could be larger than outer, in which case we need to make sure these are swapped.
+	TempCCRStateRenderThread->Inner = FMath::Min<float>(Outer, Inner);
+	TempCCRStateRenderThread->Outer = FMath::Max<float>(Outer, Inner);
+
+	if (TempCCRStateRenderThread->Inner == TempCCRStateRenderThread->Outer)
+	{
+		TempCCRStateRenderThread->Inner -= 0.0001;
+	}
+
+	TempCCRStateRenderThread->Falloff = Falloff;
+	TempCCRStateRenderThread->Invert = Invert;
+	TempCCRStateRenderThread->TemperatureType = TemperatureType;
+	TempCCRStateRenderThread->Temperature = Temperature;
+	TempCCRStateRenderThread->Tint = Tint;
+	TempCCRStateRenderThread->ColorGradingSettings = ColorGradingSettings;
+	TempCCRStateRenderThread->bEnablePerActorCC = bEnablePerActorCC;
+	TempCCRStateRenderThread->PerActorColorCorrection = PerActorColorCorrection;
+
+	GetActorBounds(true, TempCCRStateRenderThread->BoxOrigin, TempCCRStateRenderThread->BoxExtent);
+	TempCCRStateRenderThread->ActorLocation = (FVector3f)GetActorLocation();
+	TempCCRStateRenderThread->ActorRotation = (FVector3f)GetActorRotation().Euler();
+	TempCCRStateRenderThread->ActorScale = (FVector3f)GetActorScale();
+
+	// Transfer Stencil Ids.
+	{
+
 		for (const TSoftObjectPtr<AActor>& StencilActor : AffectedActors)
 		{
 			if (!StencilActor.IsValid())
@@ -157,15 +189,30 @@ void AColorCorrectRegion::TransferStencilIds()
 			{
 				if (PrimitiveComponent->bRenderCustomDepth)
 				{
-					TempStencilIds.Add(static_cast<uint32>(PrimitiveComponent->CustomDepthStencilValue));
+					TempCCRStateRenderThread->StencilIds.Add(static_cast<uint32>(PrimitiveComponent->CustomDepthStencilValue));
 				}
 			}
 		}
+	}
+
+	// Display Cluster uses HiddenPrimitives to hide Primitive components from view. 
+	// Store component id to be used on render thread.
+	if (const UStaticMeshComponent* FirstMeshComponent = FindComponentByClass<UStaticMeshComponent>())
+	{
+		if (!(TempCCRStateRenderThread->FirstPrimitiveId == FirstMeshComponent->ComponentId))
 		{
-			FScopeLock RegionScopeLock(&StencilIdsCriticalSection);
-			StencilIds = MoveTemp(TempStencilIds);
+			TempCCRStateRenderThread->FirstPrimitiveId = FirstMeshComponent->ComponentId;
 		}
 	}
+
+	{
+		ENQUEUE_RENDER_COMMAND(CopyCCProxy)([this, CCRStateToCopy = MoveTemp(TempCCRStateRenderThread)](FRHICommandListImmediate& RHICmdList)
+			{
+				ColorCorrectRenderProxy = CCRStateToCopy;
+			}
+		);
+	}
+
 }
 
 void AColorCorrectRegion::HandleAffectedActorsPropertyChange()
@@ -299,7 +346,6 @@ void AColorCorrectRegion::PostEditChangeProperty(struct FPropertyChangedEvent& P
 			ColorCorrectRegionsSubsystem->OnLevelsChanged();
 		}
 	}
-	GetActorBounds(true, BoxOrigin, BoxExtent);
 }
 #endif //WITH_EDITOR
 
diff --git a/Engine/Plugins/Experimental/ColorCorrectRegions/Source/ColorCorrectRegions/Private/ColorCorrectRegionSubsystem.cpp b/Engine/Plugins/Experimental/ColorCorrectRegions/Source/ColorCorrectRegions/Private/ColorCorrectRegionSubsystem.cpp
--- a/Engine/Plugins/Experimental/ColorCorrectRegions/Source/ColorCorrectRegions/Private/ColorCorrectRegionSubsystem.cpp
+++ b/Engine/Plugins/Experimental/ColorCorrectRegions/Source/ColorCorrectRegions/Private/ColorCorrectRegionSubsystem.cpp
@@ -195,9 +195,8 @@ void UColorCorrectRegionsSubsystem::SortRegionsByDistance(const FVector& ViewLoc
 	TMap<AColorCorrectRegion*, double> DistanceMap;
 	for (AColorCorrectRegion* Region : RegionsDistanceBased)
 	{
-		FVector BoxCenter, BoxExtents;
-		Region->GetBounds(BoxCenter, BoxExtents);
-		FVector CameraToRegionVec = (BoxCenter - ViewLocation);
+		FColorCorrectRenderProxyPtr State = Region->GetCCProxy_RenderThread();
+		FVector CameraToRegionVec = (State->BoxOrigin - ViewLocation);
 		DistanceMap.Add(Region, CameraToRegionVec.Dot(CameraToRegionVec));
 	}
 
diff --git a/Engine/Plugins/Experimental/ColorCorrectRegions/Source/ColorCorrectRegions/Private/ColorCorrectRegionsSceneViewExtension.cpp b/Engine/Plugins/Experimental/ColorCorrectRegions/Source/ColorCorrectRegions/Private/ColorCorrectRegionsSceneViewExtension.cpp
--- a/Engine/Plugins/Experimental/ColorCorrectRegions/Source/ColorCorrectRegions/Private/ColorCorrectRegionsSceneViewExtension.cpp
+++ b/Engine/Plugins/Experimental/ColorCorrectRegions/Source/ColorCorrectRegions/Private/ColorCorrectRegionsSceneViewExtension.cpp
@@ -237,7 +237,6 @@ namespace
 
 	void StencilMerger
 		( FRDGBuilder& GraphBuilder
-		, const AColorCorrectRegion* Region
 		, const FGlobalShaderMap* GlobalShaderMap
 		, const FScreenPassRenderTarget& SceneColorRenderTarget
 		, const FSceneView& View
@@ -341,29 +340,21 @@ namespace
 		SCOPED_GPU_STAT(GraphBuilder.RHICmdList, ColorCorrectRegion);
 		FRHIDepthStencilState* DepthStencilState = FScreenPassPipelineState::FDefaultDepthStencilState::GetRHI();
 
+		FColorCorrectRenderProxyPtr RegionState = Region->GetCCProxy_RenderThread();
+
 		/* If Region is pending for kill, invisible or disabled we don't need to render it.
 		*	If Region's Primitive is not visible in the current view's scene then we don't need to render it either.
 		*	We are checking if the region belongs to the same world as the view.
-		* Alternative is to get all component ids from ViewFamily.Scene and compare with actor's
-		*	Region->ActiveMeshComponent->ComponentId ComponentIdSearchTable.Contains(RegionPrimitiveComponentId)
 		*/
-		if (!IsValid(Region) ||
+		if (!RegionState->bIsActiveThisFrame ||
 			Region->IsActorBeingDestroyed() ||
-			!Region->Enabled ||
-			Region->IsHidden() ||
-#if WITH_EDITOR
-			Region->IsHiddenEd() ||
-#endif
-			Region->GetWorld() != ViewFamily.Scene->GetWorld())
+			RegionState->World != ViewFamily.Scene->GetWorld())
 		{
 			return false;
 		}
 
-		FVector BoxCenter, BoxExtents;
-		Region->GetBounds(BoxCenter, BoxExtents);
-
 		// If bounding box is zero, then we don't need to do anything.
-		if (BoxExtents.IsNearlyZero())
+		if (RegionState->BoxExtent.IsNearlyZero())
 		{
 			return false;
 		}
@@ -373,14 +364,14 @@ namespace
 		float MaxDepth = -BIG_NUMBER;
 		float MinDepth = BIG_NUMBER;
 
-		if (Region->Invert)
+		if (RegionState->Invert)
 		{
 			// In case of Region inversion we would to render the entire screen
 			Viewport = PrimaryViewRect;
 		}
 		else
 		{
-			GetPixelSpaceBoundingRect(View, BoxCenter, BoxExtents, Viewport, MaxDepth, MinDepth);
+			GetPixelSpaceBoundingRect(View, RegionState->BoxOrigin, RegionState->BoxExtent, Viewport, MaxDepth, MinDepth);
 
 			// Check if CCR is too small to be rendered (less than one pixel on the screen).
 			if (Viewport.Width() == 0 || Viewport.Height() == 0)
@@ -390,10 +381,10 @@ namespace
 
 			// This is to handle corner cases when user has a very long disproportionate region and gets either
 			// within bounds or close to the center.
-			float MaxBoxExtent = FMath::Abs(BoxExtents.GetMax());
+			float MaxBoxExtent = FMath::Abs(RegionState->BoxExtent.GetMax());
 			if (MaxDepth >= 0 && MinDepth < 0)
 			{
-				UpdateMinMaxWithFrustrumAABBIntersection(View, BoxCenter, BoxExtents, Viewport, MaxDepth);
+				UpdateMinMaxWithFrustrumAABBIntersection(View, RegionState->BoxOrigin, RegionState->BoxExtent, Viewport, MaxDepth);
 			}
 
 			FIntRect ConstrainedViewRect = View.UnscaledViewRect;
@@ -422,9 +413,9 @@ namespace
 
 		const FVector4 One(1., 1., 1., 1.);
 		const FVector4 Zero(0., 0., 0., 0.);
-		TArray<FColorGradePerRangeSettings*> AdvancedSettings{ &Region->ColorGradingSettings.Shadows,
-																&Region->ColorGradingSettings.Midtones,
-																&Region->ColorGradingSettings.Highlights };
+		TArray<const FColorGradePerRangeSettings*> AdvancedSettings{ &RegionState->ColorGradingSettings.Shadows,
+																&RegionState->ColorGradingSettings.Midtones,
+																&RegionState->ColorGradingSettings.Highlights };
 
 		// Check if any of the regions are advanced.
 		for (auto SettingsIt = AdvancedSettings.CreateConstIterator(); SettingsIt; ++SettingsIt)
@@ -457,25 +448,25 @@ namespace
 		const float DefaultTint = 0;
 
 		// If temperature is default we don't want to do the calculations.
-		FColorCorrectRegionMaterialPS::ETemperatureType TemperatureType = FMath::IsNearlyEqual(Region->Temperature, DefaultTemperature) && FMath::IsNearlyEqual(Region->Tint, DefaultTint)
+		FColorCorrectRegionMaterialPS::ETemperatureType TemperatureType = FMath::IsNearlyEqual(RegionState->Temperature, DefaultTemperature) && FMath::IsNearlyEqual(RegionState->Tint, DefaultTint)
 			? FColorCorrectRegionMaterialPS::ETemperatureType::Disabled
-			: static_cast<FColorCorrectRegionMaterialPS::ETemperatureType>(Region->TemperatureType);
+			: static_cast<FColorCorrectRegionMaterialPS::ETemperatureType>(RegionState->TemperatureType);
 
 		FScreenPassRenderTarget MergedStencilRenderTarget;
-		if (Region->bEnablePerActorCC)
+		if (RegionState->bEnablePerActorCC)
 		{
-			TArray<uint32> StencilIds = Region->GetCopyStencilIds();
-			StencilMerger(GraphBuilder, Region, GlobalShaderMap, SceneColorRenderTarget, View, SceneTextureViewportParams, RegionViewport, SceneTextures, StencilIds, MergedStencilRenderTarget);
+			TArray<uint32> StencilIds = RegionState->StencilIds;
+			StencilMerger(GraphBuilder, GlobalShaderMap, SceneColorRenderTarget, View, SceneTextureViewportParams, RegionViewport, SceneTextures, StencilIds, MergedStencilRenderTarget);
 		}
 
 		TShaderRef<FColorCorrectGenericPS> PixelShader;
 		if (AColorCorrectionWindow* CCWindow = Cast<AColorCorrectionWindow>(Region))
 		{
-			PixelShader = GetWindowShader(GlobalShaderMap, CCWindow->WindowType, TemperatureType, bIsAdvanced, MergedStencilRenderTarget.IsValid());
+			PixelShader = GetWindowShader(GlobalShaderMap, RegionState->WindowType, TemperatureType, bIsAdvanced, MergedStencilRenderTarget.IsValid());
 		}
 		else
 		{
-			PixelShader = GetRegionShader(GlobalShaderMap, Region->Type, TemperatureType, bIsAdvanced, MergedStencilRenderTarget.IsValid());
+			PixelShader = GetRegionShader(GlobalShaderMap, RegionState->Type, TemperatureType, bIsAdvanced, MergedStencilRenderTarget.IsValid());
 		}
 
 		if (MergedStencilRenderTarget.IsValid())
@@ -493,32 +484,29 @@ namespace
 		
 		// Setting constant buffer data to be passed to the shader.
 		{
-			RegionData.Rotate = FMath::DegreesToRadians<FVector3f>((FVector3f)Region->GetActorRotation().Euler());
-			RegionData.Translate = (FVector3f)Region->GetActorLocation();
+			RegionData.Rotate = FMath::DegreesToRadians<FVector3f>(RegionState->ActorRotation);
+			RegionData.Translate = RegionState->ActorLocation;
 
-			const float ScaleMultiplier = 50.;
+			const float ScaleMultiplier = View.WorldToMetersScale / 2.;
 			// Pre multiplied scale. 
-			RegionData.Scale = (FVector3f)Region->GetActorScale() * ScaleMultiplier;
-
-			RegionData.WhiteTemp = Region->Temperature;
-			RegionData.Tint = Region->Tint;
-			// Inner could be larger than outer, in which case we need to make sure these are swapped.
-			RegionData.Inner = FMath::Min<float>(Region->Outer, Region->Inner);
-			RegionData.Outer = FMath::Max<float>(Region->Outer, Region->Inner);
-			if (RegionData.Inner == RegionData.Outer)
-			{
-				RegionData.Inner -= 0.0001;
-			}
-			RegionData.Falloff = Region->Falloff;
-			RegionData.Intensity = Region->Intensity;
-			RegionData.Invert = Region->Invert;
-			RegionData.ExcludeStencil = static_cast<uint32>(Region->PerActorColorCorrection);
+			RegionData.Scale = (FVector3f)RegionState->ActorScale * ScaleMultiplier;
 
-			CCBase.ColorSaturation = (FVector4f)Region->ColorGradingSettings.Global.Saturation;
-			CCBase.ColorContrast = (FVector4f)Region->ColorGradingSettings.Global.Contrast;
-			CCBase.ColorGamma = (FVector4f)Region->ColorGradingSettings.Global.Gamma;
-			CCBase.ColorGain = (FVector4f)Region->ColorGradingSettings.Global.Gain;
-			CCBase.ColorOffset = (FVector4f)Region->ColorGradingSettings.Global.Offset;
+			RegionData.WhiteTemp = RegionState->Temperature;
+			RegionData.Tint = RegionState->Tint;
+
+			RegionData.Inner = RegionState->Inner;
+			RegionData.Outer = RegionState->Outer;
+
+			RegionData.Falloff = RegionState->Falloff;
+			RegionData.Intensity = RegionState->Intensity;
+			RegionData.Invert = RegionState->Invert;
+			RegionData.ExcludeStencil = static_cast<uint32>(RegionState->PerActorColorCorrection);
+
+			CCBase.ColorSaturation = (FVector4f)RegionState->ColorGradingSettings.Global.Saturation;
+			CCBase.ColorContrast = (FVector4f)RegionState->ColorGradingSettings.Global.Contrast;
+			CCBase.ColorGamma = (FVector4f)RegionState->ColorGradingSettings.Global.Gamma;
+			CCBase.ColorGain = (FVector4f)RegionState->ColorGradingSettings.Global.Gain;
+			CCBase.ColorOffset = (FVector4f)RegionState->ColorGradingSettings.Global.Offset;
 
 			// Set advanced 
 			if (bIsAdvanced)
@@ -526,25 +514,25 @@ namespace
 				const float GammaMin = 0.02;
 				const float GammaMax = 10.;
 				//clamp(ExternalExpressions.ColorGammaHighlights, 0.02, 10.)
-				CCShadows.ColorSaturation = (FVector4f)Region->ColorGradingSettings.Shadows.Saturation;
-				CCShadows.ColorContrast = (FVector4f)Region->ColorGradingSettings.Shadows.Contrast;
-				CCShadows.ColorGamma = (FVector4f)Clamp(Region->ColorGradingSettings.Shadows.Gamma, GammaMin, GammaMax);
-				CCShadows.ColorGain = (FVector4f)Region->ColorGradingSettings.Shadows.Gain;
-				CCShadows.ColorOffset = (FVector4f)Region->ColorGradingSettings.Shadows.Offset;
-				CCShadows.ShadowMax = Region->ColorGradingSettings.ShadowsMax;
-
-				CCMidtones.ColorSaturation = (FVector4f)Region->ColorGradingSettings.Midtones.Saturation;
-				CCMidtones.ColorContrast = (FVector4f)Region->ColorGradingSettings.Midtones.Contrast;
-				CCMidtones.ColorGamma = (FVector4f)Clamp(Region->ColorGradingSettings.Midtones.Gamma, GammaMin, GammaMax);
-				CCMidtones.ColorGain = (FVector4f)Region->ColorGradingSettings.Midtones.Gain;
-				CCMidtones.ColorOffset = (FVector4f)Region->ColorGradingSettings.Midtones.Offset;
-
-				CCHighlights.ColorSaturation = (FVector4f)Region->ColorGradingSettings.Highlights.Saturation;
-				CCHighlights.ColorContrast = (FVector4f)Region->ColorGradingSettings.Highlights.Contrast;
-				CCHighlights.ColorGamma = (FVector4f)Clamp(Region->ColorGradingSettings.Highlights.Gamma, GammaMin, GammaMax);
-				CCHighlights.ColorGain = (FVector4f)Region->ColorGradingSettings.Highlights.Gain;
-				CCHighlights.ColorOffset = (FVector4f)Region->ColorGradingSettings.Highlights.Offset;
-				CCHighlights.HighlightsMin = Region->ColorGradingSettings.HighlightsMin;
+				CCShadows.ColorSaturation = (FVector4f)RegionState->ColorGradingSettings.Shadows.Saturation;
+				CCShadows.ColorContrast = (FVector4f)RegionState->ColorGradingSettings.Shadows.Contrast;
+				CCShadows.ColorGamma = (FVector4f)Clamp(RegionState->ColorGradingSettings.Shadows.Gamma, GammaMin, GammaMax);
+				CCShadows.ColorGain = (FVector4f)RegionState->ColorGradingSettings.Shadows.Gain;
+				CCShadows.ColorOffset = (FVector4f)RegionState->ColorGradingSettings.Shadows.Offset;
+				CCShadows.ShadowMax = RegionState->ColorGradingSettings.ShadowsMax;
+
+				CCMidtones.ColorSaturation = (FVector4f)RegionState->ColorGradingSettings.Midtones.Saturation;
+				CCMidtones.ColorContrast = (FVector4f)RegionState->ColorGradingSettings.Midtones.Contrast;
+				CCMidtones.ColorGamma = (FVector4f)Clamp(RegionState->ColorGradingSettings.Midtones.Gamma, GammaMin, GammaMax);
+				CCMidtones.ColorGain = (FVector4f)RegionState->ColorGradingSettings.Midtones.Gain;
+				CCMidtones.ColorOffset = (FVector4f)RegionState->ColorGradingSettings.Midtones.Offset;
+
+				CCHighlights.ColorSaturation = (FVector4f)RegionState->ColorGradingSettings.Highlights.Saturation;
+				CCHighlights.ColorContrast = (FVector4f)RegionState->ColorGradingSettings.Highlights.Contrast;
+				CCHighlights.ColorGamma = (FVector4f)Clamp(RegionState->ColorGradingSettings.Highlights.Gamma, GammaMin, GammaMax);
+				CCHighlights.ColorGain = (FVector4f)RegionState->ColorGradingSettings.Highlights.Gain;
+				CCHighlights.ColorOffset = (FVector4f)RegionState->ColorGradingSettings.Highlights.Offset;
+				CCHighlights.HighlightsMin = RegionState->ColorGradingSettings.HighlightsMin;
 			}
 
 		}
@@ -677,11 +665,13 @@ FColorCorrectRegionsSceneViewExtension::FColorCorrectRegionsSceneViewExtension(c
 void FColorCorrectRegionsSceneViewExtension::PrePostProcessPass_RenderThread(FRDGBuilder& GraphBuilder, const FSceneView& View, const FPostProcessingInputs& Inputs)
 {
 	// Necessary for when an actor is added or removed from the scene. Also when priority is changed.
-	FScopeLock RegionScopeLock(&WorldSubsystem->RegionAccessCriticalSection);
-
-	if ((WorldSubsystem->RegionsPriorityBased.Num() == 0 && WorldSubsystem->RegionsDistanceBased.Num() == 0)|| !ViewSupportsRegions(View))
 	{
-		return;
+		FScopeLock RegionScopeLock(&WorldSubsystem->RegionAccessCriticalSection);
+
+		if ((WorldSubsystem->RegionsPriorityBased.Num() == 0 && WorldSubsystem->RegionsDistanceBased.Num() == 0) || !ViewSupportsRegions(View))
+		{
+			return;
+		}
 	}
 
 	Inputs.Validate();
@@ -739,42 +729,45 @@ void FColorCorrectRegionsSceneViewExtension::PrePostProcessPass_RenderThread(FRD
 		check(View.bIsViewInfo);
 		FSceneTextureShaderParameters SceneTextures = CreateSceneTextureShaderParameters(GraphBuilder, ((const FViewInfo&)View).GetSceneTexturesChecked(), View.GetFeatureLevel(), ESceneTextureSetupMode::All);
 
-		for (auto It = WorldSubsystem->RegionsPriorityBased.CreateConstIterator(); It; ++It)
-		{
-			AColorCorrectRegion* Region = *It;
-			RenderRegion(GraphBuilder
-				, View
-				, Inputs
-				, ViewFamily
-				, Region
-				, PrimaryViewRect
-				, SceneColorRenderTarget
-				, ScreenPercentage
-				, BackBufferRenderTarget
-				, SceneTextureViewportParams
-				, SceneTextureInput
-				, SceneTextures
-				, GlobalShaderMap
-				, DefaultBlendState);
-		}
 		WorldSubsystem->SortRegionsByDistance(View.ViewLocation);
-		for (auto It = WorldSubsystem->RegionsDistanceBased.CreateConstIterator(); It; ++It)
 		{
-			AColorCorrectRegion* Region = *It;
-			RenderRegion(GraphBuilder
-				, View
-				, Inputs
-				, ViewFamily
-				, Region
-				, PrimaryViewRect
-				, SceneColorRenderTarget
-				, ScreenPercentage
-				, BackBufferRenderTarget
-				, SceneTextureViewportParams
-				, SceneTextureInput
-				, SceneTextures
-				, GlobalShaderMap
-				, DefaultBlendState);
+			FScopeLock RegionScopeLock(&WorldSubsystem->RegionAccessCriticalSection);
+			for (auto It = WorldSubsystem->RegionsPriorityBased.CreateConstIterator(); It; ++It)
+			{
+				AColorCorrectRegion* Region = *It;
+				RenderRegion(GraphBuilder
+					, View
+					, Inputs
+					, ViewFamily
+					, Region
+					, PrimaryViewRect
+					, SceneColorRenderTarget
+					, ScreenPercentage
+					, BackBufferRenderTarget
+					, SceneTextureViewportParams
+					, SceneTextureInput
+					, SceneTextures
+					, GlobalShaderMap
+					, DefaultBlendState);
+			}
+			for (auto It = WorldSubsystem->RegionsDistanceBased.CreateConstIterator(); It; ++It)
+			{
+				AColorCorrectRegion* Region = *It;
+				RenderRegion(GraphBuilder
+					, View
+					, Inputs
+					, ViewFamily
+					, Region
+					, PrimaryViewRect
+					, SceneColorRenderTarget
+					, ScreenPercentage
+					, BackBufferRenderTarget
+					, SceneTextureViewportParams
+					, SceneTextureInput
+					, SceneTextures
+					, GlobalShaderMap
+					, DefaultBlendState);
+			}
 		}
 	}
 }
diff --git a/Engine/Plugins/Experimental/ColorCorrectRegions/Source/ColorCorrectRegions/Public/ColorCorrectRegion.h b/Engine/Plugins/Experimental/ColorCorrectRegions/Source/ColorCorrectRegions/Public/ColorCorrectRegion.h
--- a/Engine/Plugins/Experimental/ColorCorrectRegions/Source/ColorCorrectRegions/Public/ColorCorrectRegion.h
+++ b/Engine/Plugins/Experimental/ColorCorrectRegions/Source/ColorCorrectRegions/Public/ColorCorrectRegion.h
@@ -6,7 +6,6 @@
 #include "CoreMinimal.h"
 #include "Engine/Scene.h"
 #include "GameFramework/Actor.h"
-#include "Templates/SharedPointer.h"
 #include "UObject/ObjectMacros.h"
 #include "ColorCorrectRegion.generated.h"
 
@@ -14,6 +13,7 @@
 class UColorCorrectRegionsSubsystem;
 class UBillboardComponent;
 
+
 UENUM(BlueprintType)
 enum class EColorCorrectRegionsType : uint8 
 {
@@ -24,6 +24,14 @@ enum class EColorCorrectRegionsType : uint8
 	MAX
 };
 
+UENUM(BlueprintType)
+enum class EColorCorrectWindowType : uint8
+{
+	Square		UMETA(DisplayName = "Square"),
+	Circle		UMETA(DisplayName = "Circle"),
+	MAX
+};
+
 UENUM(BlueprintType)
 enum class EColorCorrectRegionTemperatureType : uint8
 {
@@ -47,6 +55,47 @@ public:
 	uint32 AssignedStencil = 0;
 };
 
+/** A state to store a copy of CCR properties to be used on render thread. */
+struct FColorCorrectRenderProxy
+{
+	EColorCorrectRegionsType Type;
+	EColorCorrectWindowType WindowType;
+
+	/** CCR Properties */
+	int32 Priority;
+	float Intensity;
+	float Inner;
+	float Outer;
+	float Falloff;
+	bool Invert;
+	EColorCorrectRegionTemperatureType TemperatureType;
+	float Temperature;
+	float Tint;
+	FColorGradingSettings ColorGradingSettings;
+
+	bool bEnablePerActorCC;
+	/**
+	* A list of actors that need to be excluded or included from color correction to be used on render thread.
+	* This list is populated during engine tick and populated with Stencil Ids from Primitive Components of AffectedActors.
+	* The main reason for this is accessing Primitive Components is not thread-safe and should be accessed on Game thread.
+	*/
+	TArray<uint32> StencilIds;
+	EColorCorrectRegionStencilType PerActorColorCorrection;
+
+	TWeakObjectPtr<UWorld> World;
+	FVector BoxOrigin;
+	FVector BoxExtent;
+
+	FVector3f ActorLocation;
+	FVector3f ActorRotation;
+	FVector3f ActorScale;
+
+	FPrimitiveComponentId FirstPrimitiveId;
+
+	bool bIsActiveThisFrame;
+};
+
+typedef TSharedPtr<FColorCorrectRenderProxy, ESPMode::ThreadSafe> FColorCorrectRenderProxyPtr;
 
 /**
  * An instance of Color Correction Region. Used to aggregate all active regions.
@@ -165,27 +214,14 @@ public:
 	void Cleanup();
 
 	/**
-	* This is used on render thread, and not atomic on purpose to avoid stalling Render thread even for a little bit. 
+	* Gets a full state for rendering. 
 	*/
-	void GetBounds(FVector& InOutOrigin, FVector& InOutBoxExtent) const
+	FColorCorrectRenderProxyPtr GetCCProxy_RenderThread()
 	{
-		InOutOrigin = BoxOrigin;
-		InOutBoxExtent = BoxExtent;
+		check(IsInRenderingThread());
+		return ColorCorrectRenderProxy;
 	};
 
-	/**
-	* This method is used on render thread, and not atomic on purpose to avoid stalling Render thread even for a little bit.
-	*/
-	const FPrimitiveComponentId& GetFirstPrimitiveId() const
-	{
-		return FirstPrimitiveId;
-	};
-
-	TArray<uint32> GetCopyStencilIds() 
-	{ 
-		FScopeLock RegionScopeLock(&StencilIdsCriticalSection);
-		return StencilIds; 
-	};
 private:
 
 #if WITH_METADATA
@@ -200,29 +236,20 @@ private:
 	void HandleAffectedActorsPropertyChange();
 
 	/**
-	* Copy stencil Ids from Primitive components to the StencilId array to be consumed by Scene view extension.
+	* Copy state required for rendering to be consumed by Scene view extension.
 	*/
-	void TransferStencilIds();
+	void TransferState();
 
 private:
 	UColorCorrectRegionsSubsystem* ColorCorrectRegionsSubsystem;
 
-	FVector BoxOrigin;
-	FVector BoxExtent;
-
-	FTransform PreviousFrameTransform;
-	FPrimitiveComponentId FirstPrimitiveId;
+	/** A copy of all properties required by render thread to process this CCR. */
+	FColorCorrectRenderProxyPtr ColorCorrectRenderProxy;
 
 	bool bActorListIsDirty;
 	uint8 ActorListChangeType;
 
-	/**
-	* A list of actors that need to be excluded or included from color correction to be used on render thread.
-	* This list is populated during engine tick and populated with Stencil Ids from Primitive Components of AffectedActors.
-	* The main reason for this is acessing Primitive Components is not thread safe and is better to be accessed on Game thread.
-	*/
-	TArray<uint32> StencilIds;
-	FCriticalSection StencilIdsCriticalSection;
+	FCriticalSection StateCopyCriticalSecion;
 
 	// This is for optimization purposes that would let us check assigned actors component's stencil ids ever few once in a while.
 	float TimeWaited = 0;
diff --git a/Engine/Plugins/Experimental/ColorCorrectRegions/Source/ColorCorrectRegions/Public/ColorCorrectWindow.h b/Engine/Plugins/Experimental/ColorCorrectRegions/Source/ColorCorrectRegions/Public/ColorCorrectWindow.h
--- a/Engine/Plugins/Experimental/ColorCorrectRegions/Source/ColorCorrectRegions/Public/ColorCorrectWindow.h
+++ b/Engine/Plugins/Experimental/ColorCorrectRegions/Source/ColorCorrectRegions/Public/ColorCorrectWindow.h
@@ -17,14 +17,6 @@
 class UColorCorrectRegionsSubsystem;
 class UBillboardComponent;
 
-UENUM(BlueprintType)
-enum class EColorCorrectWindowType : uint8
-{
-	Square		UMETA(DisplayName = "Square"),
-	Circle		UMETA(DisplayName = "Circle"),
-	MAX
-};
-
 /**
  * A Color Correction Window that functions the same way as Color Correction Regions except it modifies anything that is behind it.
  * Color Correction Windows do not have Priority property and instead are sorted based on the distance from the camera.
