Color Correction Regions: The newly created and duplicated CCRs get incremented priority assigned. Priority increment is controlled by r.CCR.PriorityIncrementAmount

#rnx

robomerge[FN_EngineMerge] -UE5-RES

diff --git a/Engine/Plugins/Experimental/ColorCorrectRegions/Source/ColorCorrectRegions/Private/ColorCorrectRegionSubsystem.cpp b/Engine/Plugins/Experimental/ColorCorrectRegions/Source/ColorCorrectRegions/Private/ColorCorrectRegionSubsystem.cpp
--- a/Engine/Plugins/Experimental/ColorCorrectRegions/Source/ColorCorrectRegions/Private/ColorCorrectRegionSubsystem.cpp
+++ b/Engine/Plugins/Experimental/ColorCorrectRegions/Source/ColorCorrectRegions/Private/ColorCorrectRegionSubsystem.cpp
@@ -17,6 +17,11 @@
 
 #define LOCTEXT_NAMESPACE "CCR"
 
+static TAutoConsoleVariable<int32> CVarCCRPriorityIncrement(
+	TEXT("r.CCR.PriorityIncrementAmount"),
+	1,
+	TEXT("Affects the priority increment of a newly created Color Correct Region."));
+
 namespace
 {
 	bool IsRegionValid(AColorCorrectRegion* InRegion, UWorld* CurrentWorld)
@@ -30,6 +35,24 @@ namespace
 		return InRegion && InRegion->GetWorld() == CurrentWorld;
 #endif
 	}
+
+	void AssignNewPriorityIfNeeded(AColorCorrectRegion* InRegion, const TArray<AColorCorrectRegion*>& RegionsPriorityBased)
+	{
+		int32 HighestPriority = 0;
+		bool bAssignNewPriority = InRegion->Priority == 0;
+		for (const AColorCorrectRegion* Region : RegionsPriorityBased)
+		{
+			if (InRegion->Priority == Region->Priority)
+			{
+				bAssignNewPriority = true;
+			}
+			HighestPriority = HighestPriority < Region->Priority ? Region->Priority : HighestPriority;
+		}
+		if (bAssignNewPriority)
+		{
+			InRegion->Priority = HighestPriority + (HighestPriority == 0 ? 1 : FMath::Max(CVarCCRPriorityIncrement.GetValueOnAnyThread(), 1));
+		}
+	}
 }
 
 void UColorCorrectRegionsSubsystem::Initialize(FSubsystemCollectionBase& Collection)
@@ -115,6 +138,11 @@ void UColorCorrectRegionsSubsystem::OnActorSpawned(AActor* InActor)
 		// adding regions twice.
 		bool bIsDistanceBased = Cast<AColorCorrectionWindow>(InActor) != nullptr;
 		TArray<AColorCorrectRegion*>* RegionsToAddTo = bIsDistanceBased ? &RegionsDistanceBased : &RegionsPriorityBased;
+		if (!bIsDistanceBased && AsRegion->Priority == 0)
+		{
+			AssignNewPriorityIfNeeded(AsRegion, RegionsPriorityBased);
+		}
+		
 		if (!RegionsToAddTo->Contains(AsRegion))
 		{
 			RegionsToAddTo->Add(AsRegion);
@@ -125,7 +153,8 @@ void UColorCorrectRegionsSubsystem::OnActorSpawned(AActor* InActor)
 			}
 		}
 	}
-	else if (bDuplicationStarted)
+
+	if (bDuplicationStarted)
 	{
 		DuplicatedActors.Add(InActor);
 	}
@@ -156,6 +185,7 @@ void UColorCorrectRegionsSubsystem::OnActorDeleted(AActor* InActor)
 void UColorCorrectRegionsSubsystem::OnDuplicateActorsEnd()
 {
 	bDuplicationStarted = false; 
+
 #if WITH_EDITOR
 	if (GEditor)
 	{
@@ -165,7 +195,14 @@ void UColorCorrectRegionsSubsystem::OnDuplicateActorsEnd()
 
 	for (AActor* DuplicatedActor : DuplicatedActors)
 	{
-		FColorCorrectRegionsStencilManager::CleanActor(DuplicatedActor);
+		if (AColorCorrectRegion* AsRegion = Cast<AColorCorrectRegion>(DuplicatedActor))
+		{
+			AssignNewPriorityIfNeeded(AsRegion, RegionsPriorityBased);
+		}
+		else
+		{
+			FColorCorrectRegionsStencilManager::CleanActor(DuplicatedActor);
+		}
 	}
 
 	DuplicatedActors.Empty();
