Color Correct Regions: Removed CCW and CCR from hit proxy generation so that it doesn't show up in Cryptomat in MRQ. The downside is that the only way to select CCR is by clicking on its billboard.

#rnx

robomerge[FN_EngineMerge] -UE5-RES

diff --git a/Engine/Plugins/Experimental/ColorCorrectRegions/Source/ColorCorrectRegions/Private/ColorCorrectRegion.cpp b/Engine/Plugins/Experimental/ColorCorrectRegions/Source/ColorCorrectRegions/Private/ColorCorrectRegion.cpp
--- a/Engine/Plugins/Experimental/ColorCorrectRegions/Source/ColorCorrectRegions/Private/ColorCorrectRegion.cpp
+++ b/Engine/Plugins/Experimental/ColorCorrectRegions/Source/ColorCorrectRegions/Private/ColorCorrectRegion.cpp
@@ -372,6 +372,7 @@ AColorCorrectionRegion::AColorCorrectionRegion(const FObjectInitializer& ObjectI
 		MeshComponent->SetCollisionEnabled(ECollisionEnabled::QueryOnly);
 		MeshComponent->SetCollisionProfileName(TEXT("OverlapAll"));
 		MeshComponent->CastShadow = false;
+		MeshComponent->SetHiddenInGame(true);
 	}
 
 	SetMeshVisibilityForRegionType();
diff --git a/Engine/Plugins/Experimental/ColorCorrectRegions/Source/ColorCorrectRegions/Private/ColorCorrectWindow.cpp b/Engine/Plugins/Experimental/ColorCorrectRegions/Source/ColorCorrectRegions/Private/ColorCorrectWindow.cpp
--- a/Engine/Plugins/Experimental/ColorCorrectRegions/Source/ColorCorrectRegions/Private/ColorCorrectWindow.cpp
+++ b/Engine/Plugins/Experimental/ColorCorrectRegions/Source/ColorCorrectRegions/Private/ColorCorrectWindow.cpp
@@ -43,6 +43,7 @@ AColorCorrectionWindow::AColorCorrectionWindow(const FObjectInitializer& ObjectI
 		MeshComponent->SetCollisionEnabled(ECollisionEnabled::QueryOnly);
 		MeshComponent->SetCollisionProfileName(TEXT("OverlapAll"));
 		MeshComponent->CastShadow = false;
+		MeshComponent->SetHiddenInGame(true);
 	}
 	SetMeshVisibilityForWindowType();
 
