Color Correction Regions/Windows: Fixed an issue with CCRs/CCWs not reacting to HiddenPrimitives and ShowOnlyPrimitives lists.

diff --git a/Engine/Plugins/Experimental/ColorCorrectRegions/Source/ColorCorrectRegions/Private/ColorCorrectRegion.cpp b/Engine/Plugins/Experimental/ColorCorrectRegions/Source/ColorCorrectRegions/Private/ColorCorrectRegion.cpp
--- a/Engine/Plugins/Experimental/ColorCorrectRegions/Source/ColorCorrectRegions/Private/ColorCorrectRegion.cpp
+++ b/Engine/Plugins/Experimental/ColorCorrectRegions/Source/ColorCorrectRegions/Private/ColorCorrectRegion.cpp
@@ -195,7 +195,6 @@ void AColorCorrectRegion::TransferState()
 		}
 	}
 
-	// Display Cluster uses HiddenPrimitives to hide Primitive components from view. 
 	// Store component id to be used on render thread.
 	if (const UStaticMeshComponent* FirstMeshComponent = FindComponentByClass<UStaticMeshComponent>())
 	{
diff --git a/Engine/Plugins/Experimental/ColorCorrectRegions/Source/ColorCorrectRegions/Private/ColorCorrectRegionsSceneViewExtension.cpp b/Engine/Plugins/Experimental/ColorCorrectRegions/Source/ColorCorrectRegions/Private/ColorCorrectRegionsSceneViewExtension.cpp
--- a/Engine/Plugins/Experimental/ColorCorrectRegions/Source/ColorCorrectRegions/Private/ColorCorrectRegionsSceneViewExtension.cpp
+++ b/Engine/Plugins/Experimental/ColorCorrectRegions/Source/ColorCorrectRegions/Private/ColorCorrectRegionsSceneViewExtension.cpp
@@ -348,7 +348,10 @@ namespace
 		*/
 		if (!RegionState->bIsActiveThisFrame ||
 			Region->IsActorBeingDestroyed() ||
-			RegionState->World != ViewFamily.Scene->GetWorld())
+			RegionState->World != ViewFamily.Scene->GetWorld() ||
+			View.HiddenPrimitives.Contains(RegionState->FirstPrimitiveId) ||
+			(View.ShowOnlyPrimitives.IsSet() && !View.ShowOnlyPrimitives->Contains(RegionState->FirstPrimitiveId))
+			)
 		{
 			return false;
 		}
