Color Correct Regions/Windows: Added a new Color Correction Primitive component that is used to determine visibility of CCR/CCW via HiddenPrimitives and ShowOnlyPrimitives.

diff --git a/Engine/Plugins/Experimental/ColorCorrectRegions/Source/ColorCorrectRegions/Private/ColorCorrectRegion.cpp b/Engine/Plugins/Experimental/ColorCorrectRegions/Source/ColorCorrectRegions/Private/ColorCorrectRegion.cpp
--- a/Engine/Plugins/Experimental/ColorCorrectRegions/Source/ColorCorrectRegions/Private/ColorCorrectRegion.cpp
+++ b/Engine/Plugins/Experimental/ColorCorrectRegions/Source/ColorCorrectRegions/Private/ColorCorrectRegion.cpp
@@ -41,6 +41,11 @@ AColorCorrectRegion::AColorCorrectRegion(const FObjectInitializer& ObjectInitial
 	RootComponent = ObjectInitializer.CreateDefaultSubobject<USceneComponent>(this, TEXT("Root"));
 	RootComponent->SetMobility(EComponentMobility::Movable);
 
+	IdentityComponent = CreateDefaultSubobject<UColorCorrectionInvisibleComponent>("IdentityComponent");
+	IdentityComponent->SetupAttachment(RootComponent);
+	IdentityComponent->CastShadow = false;
+	IdentityComponent->SetHiddenInGame(false);
+
 #if WITH_METADATA
 	if (!Cast<AColorCorrectionWindow>(this))
 	{
@@ -196,12 +201,9 @@ void AColorCorrectRegion::TransferState()
 	}
 
 	// Store component id to be used on render thread.
-	if (const UStaticMeshComponent* FirstMeshComponent = FindComponentByClass<UStaticMeshComponent>())
+	if (!(TempCCRStateRenderThread->FirstPrimitiveId == IdentityComponent->ComponentId))
 	{
-		if (!(TempCCRStateRenderThread->FirstPrimitiveId == FirstMeshComponent->ComponentId))
-		{
-			TempCCRStateRenderThread->FirstPrimitiveId = FirstMeshComponent->ComponentId;
-		}
+		TempCCRStateRenderThread->FirstPrimitiveId = IdentityComponent->ComponentId;
 	}
 
 	{
diff --git a/Engine/Plugins/Experimental/ColorCorrectRegions/Source/ColorCorrectRegions/Public/ColorCorrectRegion.h b/Engine/Plugins/Experimental/ColorCorrectRegions/Source/ColorCorrectRegions/Public/ColorCorrectRegion.h
--- a/Engine/Plugins/Experimental/ColorCorrectRegions/Source/ColorCorrectRegions/Public/ColorCorrectRegion.h
+++ b/Engine/Plugins/Experimental/ColorCorrectRegions/Source/ColorCorrectRegions/Public/ColorCorrectRegion.h
@@ -196,6 +196,10 @@ public:
 
 #endif
 
+	/** The main purpose of this component is to determine the visibility status of this Color Correction Actor. */
+	UPROPERTY()
+	TObjectPtr<UColorCorrectionInvisibleComponent> IdentityComponent;
+
 	/** To handle play in Editor, PIE and Standalone. These methods aggregate objects in play mode similarly to 
 	* Editor methods in FColorCorrectRegionsSubsystem
 	*/
@@ -281,3 +285,14 @@ public:
 	virtual FName GetCustomIconName() const override;
 #endif
 };
+
+
+/**
+ * This component class is used to determine if Color Correction Window/Region is hidden via HiddenPrimitives/ShowOnlyPrimitivesShowOnlyPrimitives
+ */
+UCLASS(NotBlueprintable, NotPlaceable)
+class UColorCorrectionInvisibleComponent : public UPrimitiveComponent
+{
+	GENERATED_BODY()
+
+};
\ No newline at end of file
