Color Correction Regions/Windows: Fixed a race condition when new level is loaded. This was caused by an invalid pointer to the world subsystem that was owning SVE.

diff --git a/Engine/Plugins/Experimental/ColorCorrectRegions/Source/ColorCorrectRegions/Private/ColorCorrectRegionSubsystem.cpp b/Engine/Plugins/Experimental/ColorCorrectRegions/Source/ColorCorrectRegions/Private/ColorCorrectRegionSubsystem.cpp
--- a/Engine/Plugins/Experimental/ColorCorrectRegions/Source/ColorCorrectRegions/Private/ColorCorrectRegionSubsystem.cpp
+++ b/Engine/Plugins/Experimental/ColorCorrectRegions/Source/ColorCorrectRegions/Private/ColorCorrectRegionSubsystem.cpp
@@ -122,8 +122,40 @@ void UColorCorrectRegionsSubsystem::Deinitialize()
 		PostProcessSceneViewExtension->IsActiveThisFrameFunctions.Add(IsActiveFunctor);
 	}
 
-	PostProcessSceneViewExtension.Reset();
-	PostProcessSceneViewExtension = nullptr;
+
+	ENQUEUE_RENDER_COMMAND(ReleaseSVE)([this](FRHICommandListImmediate& RHICmdList)
+	{
+		for (AColorCorrectRegion* Region : RegionsPriorityBased)
+		{
+			Region->Cleanup();
+		}
+		RegionsPriorityBased.Reset();
+
+		for (AColorCorrectRegion* Region : RegionsDistanceBased)
+		{
+			Region->Cleanup();
+		}
+		RegionsDistanceBased.Reset();
+
+		// Prevent this SVE from being gathered, in case it is kept alive by a strong reference somewhere else.
+		{
+			PostProcessSceneViewExtension->IsActiveThisFrameFunctions.Empty();
+
+			FSceneViewExtensionIsActiveFunctor IsActiveFunctor;
+
+			IsActiveFunctor.IsActiveFunction = [](const ISceneViewExtension* SceneViewExtension, const FSceneViewExtensionContext& Context)
+			{
+				return TOptional<bool>(false);
+			};
+
+			PostProcessSceneViewExtension->IsActiveThisFrameFunctions.Add(IsActiveFunctor);
+		}
+
+		PostProcessSceneViewExtension->Invalidate();
+		PostProcessSceneViewExtension.Reset();
+		PostProcessSceneViewExtension = nullptr;
+	});
+	FlushRenderingCommands();
 }
 
 void UColorCorrectRegionsSubsystem::OnActorSpawned(AActor* InActor)
diff --git a/Engine/Plugins/Experimental/ColorCorrectRegions/Source/ColorCorrectRegions/Private/ColorCorrectRegionsSceneViewExtension.cpp b/Engine/Plugins/Experimental/ColorCorrectRegions/Source/ColorCorrectRegions/Private/ColorCorrectRegionsSceneViewExtension.cpp
--- a/Engine/Plugins/Experimental/ColorCorrectRegions/Source/ColorCorrectRegions/Private/ColorCorrectRegionsSceneViewExtension.cpp
+++ b/Engine/Plugins/Experimental/ColorCorrectRegions/Source/ColorCorrectRegions/Private/ColorCorrectRegionsSceneViewExtension.cpp
@@ -667,15 +667,20 @@ FColorCorrectRegionsSceneViewExtension::FColorCorrectRegionsSceneViewExtension(c
 
 void FColorCorrectRegionsSceneViewExtension::PrePostProcessPass_RenderThread(FRDGBuilder& GraphBuilder, const FSceneView& View, const FPostProcessingInputs& Inputs)
 {
-	// Necessary for when an actor is added or removed from the scene. Also when priority is changed.
+	if (IsValid(WorldSubsystem))
 	{
 		FScopeLock RegionScopeLock(&WorldSubsystem->RegionAccessCriticalSection);
 
+		// Necessary for when an actor is added or removed from the scene. Also when priority is changed.
 		if ((WorldSubsystem->RegionsPriorityBased.Num() == 0 && WorldSubsystem->RegionsDistanceBased.Num() == 0) || !ViewSupportsRegions(View))
 		{
 			return;
 		}
 	}
+	else
+	{
+		return;
+	}
 
 	Inputs.Validate();
 
@@ -774,3 +779,8 @@ void FColorCorrectRegionsSceneViewExtension::PrePostProcessPass_RenderThread(FRD
 		}
 	}
 }
+
+void FColorCorrectRegionsSceneViewExtension::Invalidate()
+{
+	WorldSubsystem = nullptr;
+}
diff --git a/Engine/Plugins/Experimental/ColorCorrectRegions/Source/ColorCorrectRegions/Public/ColorCorrectRegionsSceneViewExtension.h b/Engine/Plugins/Experimental/ColorCorrectRegions/Source/ColorCorrectRegions/Public/ColorCorrectRegionsSceneViewExtension.h
--- a/Engine/Plugins/Experimental/ColorCorrectRegions/Source/ColorCorrectRegions/Public/ColorCorrectRegionsSceneViewExtension.h
+++ b/Engine/Plugins/Experimental/ColorCorrectRegions/Source/ColorCorrectRegions/Public/ColorCorrectRegionsSceneViewExtension.h
@@ -19,7 +19,9 @@ public:
 	virtual void BeginRenderViewFamily(FSceneViewFamily& InViewFamily) override {};
 	virtual void PrePostProcessPass_RenderThread(FRDGBuilder& GraphBuilder, const FSceneView& View, const FPostProcessingInputs& Inputs) override;
 	//~ End FSceneViewExtensionBase Interface
-
+public:
+	// Called when owning subsystem needs to release this extension.
+	void Invalidate();
 private:
 	UColorCorrectRegionsSubsystem* WorldSubsystem;
 };
