Removing more dependencies on Renderer/Private

diff --git a/Engine/Plugins/Experimental/ColorCorrectRegions/Source/ColorCorrectRegions/ColorCorrectRegions.Build.cs b/Engine/Plugins/Experimental/ColorCorrectRegions/Source/ColorCorrectRegions/ColorCorrectRegions.Build.cs
--- a/Engine/Plugins/Experimental/ColorCorrectRegions/Source/ColorCorrectRegions/ColorCorrectRegions.Build.cs
+++ b/Engine/Plugins/Experimental/ColorCorrectRegions/Source/ColorCorrectRegions/ColorCorrectRegions.Build.cs
@@ -1,20 +1,16 @@
 // Copyright Epic Games, Inc. All Rights Reserved.
 
+using System.IO;
 using UnrealBuildTool;
 
 public class ColorCorrectRegions : ModuleRules
 {
 	public ColorCorrectRegions(ReadOnlyTargetRules Target) : base(Target)
 	{
-		PublicIncludePaths.AddRange(
-			new string[] {
-			}
-			);
-				
 		PrivateIncludePaths.AddRange(
 			new string[] {
 				//Used in ColorCorrectRegionsPostProcessMaterial.h
-				"../../../../Source/Runtime/Renderer/Private",
+				Path.Combine(GetModuleDirectory("Renderer"), "Private"),
 			}
 		);
 		
@@ -49,11 +45,5 @@ public class ColorCorrectRegions : ModuleRules
 			PublicDependencyModuleNames.Add("UnrealEd");
 			PrivateDependencyModuleNames.Add("UnrealEd");
 		}
-
-		DynamicallyLoadedModuleNames.AddRange(
-			new string[]
-			{
-			}
-		);
 	}
 }
diff --git a/Engine/Plugins/Experimental/ColorCorrectRegions/Source/ColorCorrectRegions/Private/ColorCorrectRegionsPostProcessMaterial.cpp b/Engine/Plugins/Experimental/ColorCorrectRegions/Source/ColorCorrectRegions/Private/ColorCorrectRegionsPostProcessMaterial.cpp
--- a/Engine/Plugins/Experimental/ColorCorrectRegions/Source/ColorCorrectRegions/Private/ColorCorrectRegionsPostProcessMaterial.cpp
+++ b/Engine/Plugins/Experimental/ColorCorrectRegions/Source/ColorCorrectRegions/Private/ColorCorrectRegionsPostProcessMaterial.cpp
@@ -1,19 +1,6 @@
 // Copyright Epic Games, Inc. All Rights Reserved.
 
 #include "ColorCorrectRegionsPostProcessMaterial.h"
-#include "MaterialShader.h"
-#include "MaterialShaderType.h"
-#include "Materials/Material.h"
-#include "GlobalShader.h"
-#include "TextureResource.h"
-#include "ShaderCore.h"
-#include "ShaderParameterMacros.h"
-#include "CoreMinimal.h"
-#include "CommonRenderResources.h"
-#include "RenderGraph.h"
-#include "RendererInterface.h"
-#include "SceneRendering.h"
-#include "PostProcess/PostProcessMaterial.h"
 
 IMPLEMENT_GLOBAL_SHADER_PARAMETER_STRUCT(FCCRRegionDataInputParameter, "RegionData");
 IMPLEMENT_GLOBAL_SHADER_PARAMETER_STRUCT(FCCRColorCorrectParameter, "ColorCorrectBase");
diff --git a/Engine/Plugins/Experimental/ColorCorrectRegions/Source/ColorCorrectRegions/Private/ColorCorrectRegionsSceneViewExtension.cpp b/Engine/Plugins/Experimental/ColorCorrectRegions/Source/ColorCorrectRegions/Private/ColorCorrectRegionsSceneViewExtension.cpp
--- a/Engine/Plugins/Experimental/ColorCorrectRegions/Source/ColorCorrectRegions/Private/ColorCorrectRegionsSceneViewExtension.cpp
+++ b/Engine/Plugins/Experimental/ColorCorrectRegions/Source/ColorCorrectRegions/Private/ColorCorrectRegionsSceneViewExtension.cpp
@@ -9,12 +9,13 @@
 #include "DynamicResolutionState.h"
 #include "Engine/World.h"
 #include "EngineUtils.h"
-#include "PostProcess/PostProcessing.h"
 #include "RHI.h"
-#include "SceneRendering.h"
 #include "ScreenPass.h"
 #include "SceneView.h"
 
+#include "PostProcess/PostProcessing.h"
+#include "SceneRendering.h"
+
 // Set this to 1 to clip pixels outside of bounding box.
 #define CLIP_PIXELS_OUTSIDE_AABB 1
 
@@ -284,47 +285,27 @@ namespace
 			Parameters->PostProcessOutput = SceneTextureViewportParams;
 			Parameters->View = View.ViewUniformBuffer;
 
-			FRHIBlendState* DefaultBlendState = FScreenPassPipelineState::FDefaultBlendState::GetRHI();
-				
 			FRHIResourceCreateInfo CreateInfo(TEXT("CCR_StencilIdBuffer"));
 
 			Parameters->StencilIds = GraphBuilder.CreateSRV(FRDGBufferSRVDesc(CreateStructuredBuffer(GraphBuilder, TEXT("CCR.StencilIdBuffer"), sizeof(uint32), StencilIds.Num(), &StencilIds[0], sizeof(uint32) * StencilIds.Num())));
 			Parameters->StencilIdCount = StencilIds.Num();
 
-			{
-				GraphBuilder.AddPass(
-					RDG_EVENT_NAME("ColorCorrectRegions_StencilMerger"),
-					Parameters,
-					ERDGPassFlags::Raster,
-						[&View,
-						StencilMergerVS,
-						StencilMergerPS,
-						Parameters,
-						RegionViewport,
-						DefaultBlendState](FRHICommandList& RHICmdList)
-					{
-						check(true);
-						DrawScreenPass(
-							RHICmdList,
-							static_cast<const FViewInfo&>(View),
-							RegionViewport,
-							RegionViewport,
-							FScreenPassPipelineState(StencilMergerVS, StencilMergerPS, DefaultBlendState, FScreenPassPipelineState::FDefaultDepthStencilState::GetRHI()),
-							EScreenPassDrawFlags::None,
-							[&](FRHICommandList& RHICmdList)
-							{
-								SetShaderParameters(RHICmdList, StencilMergerPS, StencilMergerPS.GetPixelShader(), *Parameters);
-							});
-					}
-				);
-			}
+			AddDrawScreenPass(
+				GraphBuilder,
+				RDG_EVENT_NAME("ColorCorrectRegions_StencilMerger"),
+				View,
+				RegionViewport,
+				RegionViewport,
+				StencilMergerVS,
+				StencilMergerPS,
+				Parameters
+			);
 		}
 	}
 
 	bool RenderRegion
 		( FRDGBuilder& GraphBuilder
 		, const FSceneView& View
-		, const FPostProcessingInputs& Inputs
 		, const FSceneViewFamily& ViewFamily
 		, AColorCorrectRegion* Region
 		, const FIntRect& PrimaryViewRect
@@ -550,24 +531,17 @@ namespace
 			TShaderMapRef<FColorCorrectScreenPassVS> ScreenPassVS(GlobalShaderMap);
 			Parameters->RenderTargets[0] = BackBufferRenderTarget.GetRenderTargetBinding();
 
-			GraphBuilder.AddPass(
+			AddDrawScreenPass(
+				GraphBuilder,
 				RDG_EVENT_NAME("ColorCorrectRegions_ClearViewport"),
-				Parameters,
-				ERDGPassFlags::Raster,
-				[&View, ScreenPassVS, CopyPixelShader, RegionViewport, Parameters, DefaultBlendState](FRHICommandList& RHICmdList)
-				{
-					DrawScreenPass(
-						RHICmdList,
-						static_cast<const FViewInfo&>(View),
-						RegionViewport,
-						RegionViewport,
-						FScreenPassPipelineState(ScreenPassVS, CopyPixelShader, DefaultBlendState),
-						EScreenPassDrawFlags::None,
-						[&](FRHICommandList&)
-						{
-							SetShaderParameters(RHICmdList, CopyPixelShader, CopyPixelShader.GetPixelShader(), *Parameters);
-						});
-				});
+				View,
+				RegionViewport,
+				RegionViewport,
+				ScreenPassVS,
+				CopyPixelShader,
+				DefaultBlendState,
+				Parameters
+			);
 		}
 #endif
 		// Main region rendering.
@@ -592,7 +566,7 @@ namespace
 			{
 				DrawScreenPass(
 					RHICmdList,
-					static_cast<const FViewInfo&>(View),
+					View,
 					RegionViewport, // Output Viewport
 					RegionViewport, // Input Viewport
 					FScreenPassPipelineState(VertexShader, PixelShader, DefaultBlendState, DepthStencilState),
@@ -643,7 +617,7 @@ namespace
 			{
 				DrawScreenPass(
 					RHICmdList,
-					static_cast<const FViewInfo&>(View),
+					View,
 					RegionViewport,
 					RegionViewport,
 					FScreenPassPipelineState(ScreenPassVS, CopyPixelShader, CopyBlendState),
@@ -745,7 +719,6 @@ void FColorCorrectRegionsSceneViewExtension::PrePostProcessPass_RenderThread(FRD
 				AColorCorrectRegion* Region = *It;
 				RenderRegion(GraphBuilder
 					, View
-					, Inputs
 					, ViewFamily
 					, Region
 					, PrimaryViewRect
@@ -763,7 +736,6 @@ void FColorCorrectRegionsSceneViewExtension::PrePostProcessPass_RenderThread(FRD
 				AColorCorrectRegion* Region = *It;
 				RenderRegion(GraphBuilder
 					, View
-					, Inputs
 					, ViewFamily
 					, Region
 					, PrimaryViewRect
diff --git a/Engine/Plugins/Media/ImgMedia/Source/ImgMedia/ImgMedia.Build.cs b/Engine/Plugins/Media/ImgMedia/Source/ImgMedia/ImgMedia.Build.cs
--- a/Engine/Plugins/Media/ImgMedia/Source/ImgMedia/ImgMedia.Build.cs
+++ b/Engine/Plugins/Media/ImgMedia/Source/ImgMedia/ImgMedia.Build.cs
@@ -1,6 +1,5 @@
 // Copyright Epic Games, Inc. All Rights Reserved.
 
-using System;
 using System.IO;
 
 namespace UnrealBuildTool.Rules
@@ -70,8 +69,6 @@ namespace UnrealBuildTool.Rules
 				PrivateDependencyModuleNames.Add("D3D12RHI");
 			}
 			
-			PrivateIncludePaths.Add(Path.Combine(EngineDir, "Source", "Runtime", "Renderer", "Private"));
-
 			// Is this the editor?
 			if (Target.bBuildEditor == true)
 			{
diff --git a/Engine/Plugins/Media/ImgMedia/Source/ImgMedia/Private/Readers/ExrImgMediaReaderGpu.cpp b/Engine/Plugins/Media/ImgMedia/Source/ImgMedia/Private/Readers/ExrImgMediaReaderGpu.cpp
--- a/Engine/Plugins/Media/ImgMedia/Source/ImgMedia/Private/Readers/ExrImgMediaReaderGpu.cpp
+++ b/Engine/Plugins/Media/ImgMedia/Source/ImgMedia/Private/Readers/ExrImgMediaReaderGpu.cpp
@@ -19,7 +19,7 @@
 #include "GlobalShader.h"
 #include "CommonRenderResources.h"
 #include "ImgMediaLoader.h"
-#include "PostProcess/SceneFilterRendering.h"
+#include "PostProcess/DrawRectangle.h"
 #include "RHIStaticStates.h"
 #include "SceneUtils.h"
 #include "ShaderParameterUtils.h"
@@ -62,13 +62,13 @@ namespace {
 
 		EDrawRectangleFlags DrawRectangleFlags = EDRF_UseTriangleOptimization;
 
-		DrawPostProcessPass(
+		UE::Renderer::PostProcess::DrawPostProcessPass(
 			RHICmdList,
+			PipelineState.VertexShader,
 			0, 0, OutputResolution.X, OutputResolution.Y,
 			Viewport.Min.X, Viewport.Min.Y, Viewport.Width(), Viewport.Height(),
 			OutputResolution,
 			OutputResolution,
-			PipelineState.VertexShader,
 			INDEX_NONE,
 			false,
 			DrawRectangleFlags);
diff --git a/Engine/Plugins/Media/PixelStreaming/Source/PixelStreamingEditor/PixelStreamingEditor.build.cs b/Engine/Plugins/Media/PixelStreaming/Source/PixelStreamingEditor/PixelStreamingEditor.build.cs
--- a/Engine/Plugins/Media/PixelStreaming/Source/PixelStreamingEditor/PixelStreamingEditor.build.cs
+++ b/Engine/Plugins/Media/PixelStreaming/Source/PixelStreamingEditor/PixelStreamingEditor.build.cs
@@ -1,9 +1,6 @@
 // Copyright Epic Games, Inc. All Rights Reserved.
 
-using System;
 using System.IO;
-using System.Collections.Generic;
-using EpicGames.Core;
 
 namespace UnrealBuildTool.Rules
 {
@@ -17,7 +14,6 @@ namespace UnrealBuildTool.Rules
 			PrivateIncludePaths.AddRange(new string[]
 			{
 				Path.Combine(EngineDir, "Plugins/Media/PixelStreaming/Source/PixelStreaming/Private"),
-				Path.Combine(EngineDir, "Source/Runtime/Renderer/Private"),
 			});
 
 			PrivateDependencyModuleNames.AddRange(new string[]
diff --git a/Engine/Plugins/Runtime/nDisplay/Source/DisplayCluster/DisplayCluster.Build.cs b/Engine/Plugins/Runtime/nDisplay/Source/DisplayCluster/DisplayCluster.Build.cs
--- a/Engine/Plugins/Runtime/nDisplay/Source/DisplayCluster/DisplayCluster.Build.cs
+++ b/Engine/Plugins/Runtime/nDisplay/Source/DisplayCluster/DisplayCluster.Build.cs
@@ -1,7 +1,6 @@
 // Copyright Epic Games, Inc. All Rights Reserved.
 
 using UnrealBuildTool;
-using System.IO;
 
 public class DisplayCluster : ModuleRules
 {
@@ -9,11 +8,6 @@ public class DisplayCluster : ModuleRules
 	{
 		PublicDefinitions.Add("WITH_OCIO=0");
 
-		PrivateIncludePaths.AddRange(
-			new string[] {
-				System.IO.Path.Combine(GetModuleDirectory("Renderer"), "Private"), //required for FPostProcessMaterialInputs
-			});
-
 		PublicDependencyModuleNames.AddRange(
 			new string[] {
 				"CinematicCamera",
diff --git a/Engine/Plugins/Runtime/nDisplay/Source/DisplayCluster/Private/Render/Viewport/DisplayClusterViewportManagerViewExtension.cpp b/Engine/Plugins/Runtime/nDisplay/Source/DisplayCluster/Private/Render/Viewport/DisplayClusterViewportManagerViewExtension.cpp
--- a/Engine/Plugins/Runtime/nDisplay/Source/DisplayCluster/Private/Render/Viewport/DisplayClusterViewportManagerViewExtension.cpp
+++ b/Engine/Plugins/Runtime/nDisplay/Source/DisplayCluster/Private/Render/Viewport/DisplayClusterViewportManagerViewExtension.cpp
@@ -5,7 +5,6 @@
 
 #include "Render/Viewport/DisplayClusterViewportManager.h"
 #include "Render/Viewport/DisplayClusterViewportManagerProxy.h"
-#include "SceneRendering.h"
 
 #include "IDisplayCluster.h"
 #include "IDisplayClusterCallbacks.h"
@@ -42,10 +41,8 @@ void FDisplayClusterViewportManagerViewExtension::PreRenderView_RenderThread(FRD
 {
 	if (InView.bIsViewInfo)
 	{
-		FViewInfo& ViewInfo = static_cast<FViewInfo&>(InView);
-
 		// UE-145088: VR HMD device post-processing cannot be applied to nDisplay rendering
-		ViewInfo.bHMDHiddenAreaMaskActive = false;
+		InView.bHMDHiddenAreaMaskActive = false;
 	}
 }
 
diff --git a/Engine/Plugins/Runtime/nDisplay/Source/DisplayCluster/Private/Render/Viewport/DisplayClusterViewportProxy.cpp b/Engine/Plugins/Runtime/nDisplay/Source/DisplayCluster/Private/Render/Viewport/DisplayClusterViewportProxy.cpp
--- a/Engine/Plugins/Runtime/nDisplay/Source/DisplayCluster/Private/Render/Viewport/DisplayClusterViewportProxy.cpp
+++ b/Engine/Plugins/Runtime/nDisplay/Source/DisplayCluster/Private/Render/Viewport/DisplayClusterViewportProxy.cpp
@@ -42,9 +42,7 @@
 #include "ShaderParameterUtils.h"
 
 #include "ScreenRendering.h"
-#include "PostProcess/SceneFilterRendering.h"
 #include "PostProcess/PostProcessAA.h"
-#include "SceneRendering.h"
 #include "RenderGraphUtils.h"
 #include "RenderGraphResources.h"
 #include "PostProcess/DrawRectangle.h"
diff --git a/Engine/Plugins/Runtime/nDisplay/Source/DisplayClusterLightCardEditor/Private/Viewport/DisplayClusterLightCardEditorViewportClient.cpp b/Engine/Plugins/Runtime/nDisplay/Source/DisplayClusterLightCardEditor/Private/Viewport/DisplayClusterLightCardEditorViewportClient.cpp
--- a/Engine/Plugins/Runtime/nDisplay/Source/DisplayClusterLightCardEditor/Private/Viewport/DisplayClusterLightCardEditorViewportClient.cpp
+++ b/Engine/Plugins/Runtime/nDisplay/Source/DisplayClusterLightCardEditor/Private/Viewport/DisplayClusterLightCardEditorViewportClient.cpp
@@ -50,7 +50,6 @@
 #include "PreviewScene.h"
 #include "ProceduralMeshComponent.h"
 #include "RayTracingDebugVisualizationMenuCommands.h"
-#include "Runtime/Renderer/Private/SceneRendering.h"
 #include "ScopedTransaction.h"
 #include "Settings/LevelEditorViewportSettings.h"
 #include "Slate/SceneViewport.h"
diff --git a/Engine/Plugins/Runtime/nDisplay/Source/DisplayClusterMedia/DisplayClusterMedia.Build.cs b/Engine/Plugins/Runtime/nDisplay/Source/DisplayClusterMedia/DisplayClusterMedia.Build.cs
--- a/Engine/Plugins/Runtime/nDisplay/Source/DisplayClusterMedia/DisplayClusterMedia.Build.cs
+++ b/Engine/Plugins/Runtime/nDisplay/Source/DisplayClusterMedia/DisplayClusterMedia.Build.cs
@@ -1,17 +1,11 @@
 // Copyright Epic Games, Inc. All Rights Reserved.
 
 using UnrealBuildTool;
-using System.IO;
 
 public class DisplayClusterMedia : ModuleRules
 {
 	public DisplayClusterMedia(ReadOnlyTargetRules ROTargetRules) : base(ROTargetRules)
 	{
-		// [temporary] We need this to be able to use some private data types. This should
-		// be removed once we move the nD rendering pipeline to RDG.
-		string EngineDir = Path.GetFullPath(Target.RelativeEnginePath);
-		PrivateIncludePaths.Add(Path.Combine(EngineDir, "Source", "Runtime", "Renderer", "Private"));
-
 		PublicIncludePathModuleNames.AddRange(
 			new string[] {
 				"DisplayClusterShaders",
diff --git a/Engine/Source/Editor/PixelInspector/PixelInspectorModule.Build.cs b/Engine/Source/Editor/PixelInspector/PixelInspectorModule.Build.cs
--- a/Engine/Source/Editor/PixelInspector/PixelInspectorModule.Build.cs
+++ b/Engine/Source/Editor/PixelInspector/PixelInspectorModule.Build.cs
@@ -1,7 +1,6 @@
 // Copyright Epic Games, Inc. All Rights Reserved.
 
 using UnrealBuildTool;
-using System.IO;
 
 public class PixelInspectorModule : ModuleRules
 {
@@ -21,14 +20,6 @@ public class PixelInspectorModule : ModuleRules
                 "PropertyEditor",
 			}
         );
-		var EngineDir = Path.GetFullPath(Target.RelativeEnginePath);
-
-		PrivateIncludePaths.AddRange(
-			new string[] {
-				// required for PostProcessing.
-				Path.Combine(EngineDir, "Source/Runtime/Renderer/Private")
-			}
-		);
 
 		PrivateDependencyModuleNames.AddRange(
              new string[] {
diff --git a/Engine/Source/Editor/PixelInspector/Private/PixelInspector.cpp b/Engine/Source/Editor/PixelInspector/Private/PixelInspector.cpp
--- a/Engine/Source/Editor/PixelInspector/Private/PixelInspector.cpp
+++ b/Engine/Source/Editor/PixelInspector/Private/PixelInspector.cpp
@@ -17,8 +17,7 @@
 #include "PixelInspectorStyle.h"
 #include "ScreenPass.h"
 #include "UnrealClient.h"
-#include "PostProcess/PostProcessing.h"
-#include "PostProcess/PostProcessMaterial.h"
+#include "PostProcess/PostProcessMaterialInputs.h"
 
 #include "EngineGlobals.h"
 #include "EditorViewportClient.h"
diff --git a/Engine/Source/Runtime/HeadMountedDisplay/HeadMountedDisplay.Build.cs b/Engine/Source/Runtime/HeadMountedDisplay/HeadMountedDisplay.Build.cs
--- a/Engine/Source/Runtime/HeadMountedDisplay/HeadMountedDisplay.Build.cs
+++ b/Engine/Source/Runtime/HeadMountedDisplay/HeadMountedDisplay.Build.cs
@@ -7,12 +7,6 @@ public class HeadMountedDisplay : ModuleRules
 {
     public HeadMountedDisplay(ReadOnlyTargetRules Target) : base(Target)
 	{
-        PrivateIncludePaths.AddRange(
-            new string[] {
-                "Runtime/Renderer/Private"
-            }
-        );
-
         PrivateDependencyModuleNames.AddRange(
 			new string[] {
 				"Core",
diff --git a/Engine/Source/Runtime/HeadMountedDisplay/Private/DefaultXRCamera.cpp b/Engine/Source/Runtime/HeadMountedDisplay/Private/DefaultXRCamera.cpp
--- a/Engine/Source/Runtime/HeadMountedDisplay/Private/DefaultXRCamera.cpp
+++ b/Engine/Source/Runtime/HeadMountedDisplay/Private/DefaultXRCamera.cpp
@@ -2,12 +2,12 @@
 
 #include "DefaultXRCamera.h"
 #include "GameFramework/PlayerController.h"
-#include "RendererPrivate.h"
-#include "ScenePrivate.h"
 #include "Slate/SceneViewport.h"
 #include "StereoRendering.h"
 #include "StereoRenderTargetManager.h"
 #include "IHeadMountedDisplay.h"
+#include "RenderGraphBuilder.h"
+#include "SceneView.h"
 
 FDefaultXRCamera::FDefaultXRCamera(const FAutoRegister& AutoRegister, IXRTrackingSystem* InTrackingSystem, int32 InDeviceId)
 	: FHMDSceneViewExtension(AutoRegister)
