Color Correct Regions: Waiting for the rendering commands to complete before cleaning up CCR Actors.

diff --git a/Engine/Plugins/Experimental/ColorCorrectRegions/Source/ColorCorrectRegions/Private/ColorCorrectRegion.cpp b/Engine/Plugins/Experimental/ColorCorrectRegions/Source/ColorCorrectRegions/Private/ColorCorrectRegion.cpp
--- a/Engine/Plugins/Experimental/ColorCorrectRegions/Source/ColorCorrectRegions/Private/ColorCorrectRegion.cpp
+++ b/Engine/Plugins/Experimental/ColorCorrectRegions/Source/ColorCorrectRegions/Private/ColorCorrectRegion.cpp
@@ -62,7 +62,7 @@ void AColorCorrectRegion::BeginPlay()
 		ColorCorrectRegionsSubsystem = World->GetSubsystem<UColorCorrectRegionsSubsystem>();
 	}
 
-	if (ColorCorrectRegionsSubsystem)
+	if (ColorCorrectRegionsSubsystem.IsValid())
 	{
 		ColorCorrectRegionsSubsystem->OnActorSpawned(this);
 	}
@@ -70,7 +70,7 @@ void AColorCorrectRegion::BeginPlay()
 
 void AColorCorrectRegion::EndPlay(const EEndPlayReason::Type EndPlayReason)
 {
-	if (ColorCorrectRegionsSubsystem)
+	if (ColorCorrectRegionsSubsystem.IsValid())
 	{
 		ColorCorrectRegionsSubsystem->OnActorDeleted(this);
 		ColorCorrectRegionsSubsystem = nullptr;
@@ -80,7 +80,7 @@ void AColorCorrectRegion::EndPlay(const EEndPlayReason::Type EndPlayReason)
 
 void AColorCorrectRegion::BeginDestroy()
 {
-	if (ColorCorrectRegionsSubsystem)
+	if (ColorCorrectRegionsSubsystem.IsValid())
 	{
 		ColorCorrectRegionsSubsystem->OnActorDeleted(this);
 		ColorCorrectRegionsSubsystem = nullptr;
@@ -109,7 +109,7 @@ void AColorCorrectRegion::TickActor(float DeltaTime, ELevelTick TickType, FActor
 		TimeWaited += DeltaTime;
 		const float WaitTimeInSecs = 1.0;
 
-		if (!ColorCorrectRegionsSubsystem)
+		if (!ColorCorrectRegionsSubsystem.IsValid())
 		{
 			if (const UWorld* World = GetWorld())
 			{
@@ -117,7 +117,7 @@ void AColorCorrectRegion::TickActor(float DeltaTime, ELevelTick TickType, FActor
 			}
 		}
 		
-		if (ColorCorrectRegionsSubsystem && TimeWaited >= WaitTimeInSecs)
+		if (ColorCorrectRegionsSubsystem.IsValid() && TimeWaited >= WaitTimeInSecs)
 		{
 			ColorCorrectRegionsSubsystem->CheckAssignedActorsValidity(this);
 			TimeWaited = 0;
@@ -126,11 +126,6 @@ void AColorCorrectRegion::TickActor(float DeltaTime, ELevelTick TickType, FActor
 	}
 }
 
-void AColorCorrectRegion::Cleanup()
-{
-	ColorCorrectRegionsSubsystem = nullptr;
-}
-
 void AColorCorrectRegion::TransferState()
 {
 	TRACE_CPUPROFILER_EVENT_SCOPE_TEXT(*FString::Printf(TEXT("CCR.TransferState %s"), *GetName()));
@@ -246,7 +241,7 @@ void AColorCorrectRegion::HandleAffectedActorsPropertyChange()
 				}
 			}
 			bEventHandled = true;
-			if (ColorCorrectRegionsSubsystem)
+			if (ColorCorrectRegionsSubsystem.IsValid())
 			{
 				ColorCorrectRegionsSubsystem->AssignStencilIdsToPerActorCC(this);
 			}
@@ -257,7 +252,7 @@ void AColorCorrectRegion::HandleAffectedActorsPropertyChange()
 			|| ActorListChangeType == EPropertyChangeType::ValueSet)
 		{
 			bEventHandled = true;
-			if (ColorCorrectRegionsSubsystem)
+			if (ColorCorrectRegionsSubsystem.IsValid())
 			{
 				ColorCorrectRegionsSubsystem->ClearStencilIdsToPerActorCC(this);
 			}
@@ -315,7 +310,7 @@ void AColorCorrectRegion::PostEditChangeProperty(struct FPropertyChangedEvent& P
 	Super::PostEditChangeProperty(PropertyChangedEvent);
 	const FName PropertyName = PropertyChangedEvent.GetPropertyName();
 
-	if (!ColorCorrectRegionsSubsystem)
+	if (!ColorCorrectRegionsSubsystem.IsValid())
 	{
 		if (const UWorld* World = GetWorld())
 		{
@@ -335,14 +330,14 @@ void AColorCorrectRegion::PostEditChangeProperty(struct FPropertyChangedEvent& P
 	// Therefore we need to refresh priority if PropertyChangedEvent.Property is nullptr. 
 	if (PropertyName == GET_MEMBER_NAME_CHECKED(AColorCorrectRegion, Priority) || PropertyChangedEvent.Property == nullptr)
 	{
-		if (ColorCorrectRegionsSubsystem)
+		if (ColorCorrectRegionsSubsystem.IsValid())
 		{
 			ColorCorrectRegionsSubsystem->SortRegionsByPriority();
 		}
 	}
 	if (PropertyName == GET_MEMBER_NAME_CHECKED(AColorCorrectRegion, Type) || PropertyChangedEvent.Property == nullptr)
 	{
-		if (ColorCorrectRegionsSubsystem)
+		if (ColorCorrectRegionsSubsystem.IsValid())
 		{
 			ColorCorrectRegionsSubsystem->OnLevelsChanged();
 		}
diff --git a/Engine/Plugins/Experimental/ColorCorrectRegions/Source/ColorCorrectRegions/Private/ColorCorrectRegionSubsystem.cpp b/Engine/Plugins/Experimental/ColorCorrectRegions/Source/ColorCorrectRegions/Private/ColorCorrectRegionSubsystem.cpp
--- a/Engine/Plugins/Experimental/ColorCorrectRegions/Source/ColorCorrectRegions/Private/ColorCorrectRegionSubsystem.cpp
+++ b/Engine/Plugins/Experimental/ColorCorrectRegions/Source/ColorCorrectRegions/Private/ColorCorrectRegionSubsystem.cpp
@@ -96,17 +96,6 @@ void UColorCorrectRegionsSubsystem::Deinitialize()
 	}
 #endif
 	GetWorld()->OnLevelsChanged().RemoveAll(this);
-	for (AColorCorrectRegion* Region : RegionsPriorityBased)
-	{
-		Region->Cleanup();
-	}
-	RegionsPriorityBased.Reset();	
-	
-	for (AColorCorrectRegion* Region : RegionsDistanceBased)
-	{
-		Region->Cleanup();
-	}
-	RegionsDistanceBased.Reset();
 
 	// Prevent this SVE from being gathered, in case it is kept alive by a strong reference somewhere else.
 	{
@@ -122,21 +111,8 @@ void UColorCorrectRegionsSubsystem::Deinitialize()
 		PostProcessSceneViewExtension->IsActiveThisFrameFunctions.Add(IsActiveFunctor);
 	}
 
-
 	ENQUEUE_RENDER_COMMAND(ReleaseSVE)([this](FRHICommandListImmediate& RHICmdList)
 	{
-		for (AColorCorrectRegion* Region : RegionsPriorityBased)
-		{
-			Region->Cleanup();
-		}
-		RegionsPriorityBased.Reset();
-
-		for (AColorCorrectRegion* Region : RegionsDistanceBased)
-		{
-			Region->Cleanup();
-		}
-		RegionsDistanceBased.Reset();
-
 		// Prevent this SVE from being gathered, in case it is kept alive by a strong reference somewhere else.
 		{
 			PostProcessSceneViewExtension->IsActiveThisFrameFunctions.Empty();
@@ -155,7 +131,12 @@ void UColorCorrectRegionsSubsystem::Deinitialize()
 		PostProcessSceneViewExtension.Reset();
 		PostProcessSceneViewExtension = nullptr;
 	});
+
+	// Finish all rendering commands before cleaning up actors.
 	FlushRenderingCommands();
+
+	RegionsPriorityBased.Reset();
+	RegionsDistanceBased.Reset();
 }
 
 void UColorCorrectRegionsSubsystem::OnActorSpawned(AActor* InActor)
@@ -202,8 +183,6 @@ void UColorCorrectRegionsSubsystem::OnActorDeleted(AActor* InActor)
 		)
 #endif
 	{
-		AsRegion->Cleanup();
-
 #if WITH_EDITOR
 		FColorCorrectRegionsStencilManager::OnCCRRemoved(GetWorld(), AsRegion);
 #endif
diff --git a/Engine/Plugins/Experimental/ColorCorrectRegions/Source/ColorCorrectRegions/Public/ColorCorrectRegion.h b/Engine/Plugins/Experimental/ColorCorrectRegions/Source/ColorCorrectRegions/Public/ColorCorrectRegion.h
--- a/Engine/Plugins/Experimental/ColorCorrectRegions/Source/ColorCorrectRegions/Public/ColorCorrectRegion.h
+++ b/Engine/Plugins/Experimental/ColorCorrectRegions/Source/ColorCorrectRegions/Public/ColorCorrectRegion.h
@@ -211,12 +211,6 @@ public:
 	virtual void TickActor(float DeltaTime, enum ELevelTick TickType, FActorTickFunction& ThisTickFunction);
 	virtual bool ShouldTickIfViewportsOnly() const;
 
-	/** 
-	* We have to manage the lifetime of the region ourselves, because EndPlay is not guaranteed to be called 
-	* and BeginDestroy could be called from GC when it is too late.
-	*/
-	void Cleanup();
-
 	/**
 	* Gets a full state for rendering. 
 	*/
@@ -245,7 +239,7 @@ private:
 	void TransferState();
 
 private:
-	UColorCorrectRegionsSubsystem* ColorCorrectRegionsSubsystem;
+	TWeakObjectPtr<UColorCorrectRegionsSubsystem> ColorCorrectRegionsSubsystem;
 
 	/** A copy of all properties required by render thread to process this CCR. */
 	FColorCorrectRenderProxyPtr ColorCorrectRenderProxy;
