Color Correct Actors no longer have collision in editor mode.

diff --git a/Engine/Plugins/Experimental/ColorCorrectRegions/Source/ColorCorrectRegions/Private/ColorCorrectRegion.cpp b/Engine/Plugins/Experimental/ColorCorrectRegions/Source/ColorCorrectRegions/Private/ColorCorrectRegion.cpp
--- a/Engine/Plugins/Experimental/ColorCorrectRegions/Source/ColorCorrectRegions/Private/ColorCorrectRegion.cpp
+++ b/Engine/Plugins/Experimental/ColorCorrectRegions/Source/ColorCorrectRegions/Private/ColorCorrectRegion.cpp
@@ -1,6 +1,7 @@
 // Copyright Epic Games, Inc. All Rights Reserved.
 
 #include "ColorCorrectRegion.h"
+#include "Async/Async.h"
 #include "ColorCorrectRegionsModule.h"
 #include "ColorCorrectRegionsSubsystem.h"
 #include "ColorCorrectWindow.h"
@@ -8,6 +9,7 @@
 #include "Components/MeshComponent.h"
 #include "Components/StaticMeshComponent.h"
 #include "CoreMinimal.h"
+#include "Engine/CollisionProfile.h"
 #include "Engine/GameEngine.h"
 #include "Engine/StaticMesh.h"
 #include "Engine/Texture2D.h"
@@ -100,7 +102,6 @@ void AColorCorrectRegion::TickActor(float DeltaTime, ELevelTick TickType, FActor
 
 	TRACE_CPUPROFILER_EVENT_SCOPE_TEXT(*FString::Printf(TEXT("CCR.TickActor %s"), *GetName()));
 
-	HandleAffectedActorsPropertyChange();
 	TransferState();
 
 
@@ -169,7 +170,7 @@ void AColorCorrectRegion::TransferState()
 	TempCCRStateRenderThread->bEnablePerActorCC = bEnablePerActorCC;
 	TempCCRStateRenderThread->PerActorColorCorrection = PerActorColorCorrection;
 
-	GetActorBounds(true, TempCCRStateRenderThread->BoxOrigin, TempCCRStateRenderThread->BoxExtent);
+	GetActorBounds(false, TempCCRStateRenderThread->BoxOrigin, TempCCRStateRenderThread->BoxExtent);
 	TempCCRStateRenderThread->ActorLocation = (FVector3f)GetActorLocation();
 	TempCCRStateRenderThread->ActorRotation = (FVector3f)GetActorRotation().Euler();
 	TempCCRStateRenderThread->ActorScale = (FVector3f)GetActorScale();
@@ -211,51 +212,44 @@ void AColorCorrectRegion::TransferState()
 
 }
 
-void AColorCorrectRegion::HandleAffectedActorsPropertyChange()
+void AColorCorrectRegion::HandleAffectedActorsPropertyChange(uint32 ActorListChangeType)
 {
-	if (bActorListIsDirty)
+	if (ActorListChangeType == EPropertyChangeType::ArrayAdd
+		|| ActorListChangeType == EPropertyChangeType::ValueSet)
 	{
-		bool bEventHandled = false;
-		bActorListIsDirty = false;
-		if (ActorListChangeType == EPropertyChangeType::ArrayAdd
-			|| ActorListChangeType == EPropertyChangeType::ValueSet)
+		// In case user assigns Color Correct Region or Window, we should remove it as it is invalid operation.
 		{
-			// In case user assigns Color Correct Region or Window, we should remove it as it is invalid operation.
+			TArray<TSoftObjectPtr<AActor>> ActorsToRemove;
+			for (const TSoftObjectPtr<AActor>& StencilActor : AffectedActors)
 			{
-				TArray<TSoftObjectPtr<AActor>> ActorsToRemove;
-				for (const TSoftObjectPtr<AActor>& StencilActor : AffectedActors)
+				if (AColorCorrectRegion* CCRCast = Cast<AColorCorrectRegion>(StencilActor.Get()))
 				{
-					if (AColorCorrectRegion* CCRCast = Cast<AColorCorrectRegion>(StencilActor.Get()))
-					{
-						ActorsToRemove.Add(StencilActor);
-					}
-				}
-				if (ActorsToRemove.Num() > 0)
-				{
-					UE_LOG(ColorCorrectRegions, Warning, TEXT("Color Correct Region or Window assignment to Per Actor CC is not supported."));
-				}
-				for (const TSoftObjectPtr<AActor>& StencilActor : ActorsToRemove)
-				{
-					AffectedActors.Remove(StencilActor);
-					AffectedActors.FindOrAdd(TSoftObjectPtr<AActor>());
+					ActorsToRemove.Add(StencilActor);
 				}
 			}
-			bEventHandled = true;
-			if (ColorCorrectRegionsSubsystem.IsValid())
+			if (ActorsToRemove.Num() > 0)
 			{
-				ColorCorrectRegionsSubsystem->AssignStencilIdsToPerActorCC(this);
+				UE_LOG(ColorCorrectRegions, Warning, TEXT("Color Correct Region or Window assignment to Per Actor CC is not supported."));
+			}
+			for (const TSoftObjectPtr<AActor>& StencilActor : ActorsToRemove)
+			{
+				AffectedActors.Remove(StencilActor);
+				AffectedActors.FindOrAdd(TSoftObjectPtr<AActor>());
 			}
 		}
+		if (ColorCorrectRegionsSubsystem.IsValid())
+		{
+			ColorCorrectRegionsSubsystem->AssignStencilIdsToPerActorCC(this);
+		}
+	}
 
-		if (ActorListChangeType == EPropertyChangeType::ArrayClear
-			|| ActorListChangeType == EPropertyChangeType::ArrayRemove
-			|| ActorListChangeType == EPropertyChangeType::ValueSet)
+	if (ActorListChangeType == EPropertyChangeType::ArrayClear
+		|| ActorListChangeType == EPropertyChangeType::ArrayRemove
+		|| ActorListChangeType == EPropertyChangeType::ValueSet)
+	{
+		if (ColorCorrectRegionsSubsystem.IsValid())
 		{
-			bEventHandled = true;
-			if (ColorCorrectRegionsSubsystem.IsValid())
-			{
-				ColorCorrectRegionsSubsystem->ClearStencilIdsToPerActorCC(this);
-			}
+			ColorCorrectRegionsSubsystem->ClearStencilIdsToPerActorCC(this);
 		}
 	}
 }
@@ -320,9 +314,11 @@ void AColorCorrectRegion::PostEditChangeProperty(struct FPropertyChangedEvent& P
 
 	if (PropertyName == GET_MEMBER_NAME_CHECKED(AColorCorrectRegion, AffectedActors))
 	{
-		/** Since there might be Dialogs involved we need to dirty this CCR and handle the rest on Tick. */
-		ActorListChangeType = PropertyChangedEvent.ChangeType;
-		bActorListIsDirty = true;
+		/** Since there might be Dialogs involved we need to run this on game thread. */
+		AsyncTask(ENamedThreads::GameThread, [this, ActorListChangeType = PropertyChangedEvent.ChangeType]() 
+		{
+			HandleAffectedActorsPropertyChange(ActorListChangeType);
+		});
 	}
 
 	// Reorder all CCRs after the Priority property has changed.
@@ -365,8 +361,8 @@ AColorCorrectionRegion::AColorCorrectionRegion(const FObjectInitializer& ObjectI
 		MeshComponent->SetupAttachment(RootComponent);
 		MeshComponent->SetStaticMesh(StaticMeshes[static_cast<uint8>(CCRType)]);
 		MeshComponent->SetMaterial(0, Material);
-		MeshComponent->SetCollisionEnabled(ECollisionEnabled::QueryOnly);
-		MeshComponent->SetCollisionProfileName(TEXT("OverlapAll"));
+		MeshComponent->SetCollisionEnabled(ECollisionEnabled::NoCollision);
+		MeshComponent->SetCollisionProfileName(UCollisionProfile::NoCollision_ProfileName);
 		MeshComponent->CastShadow = false;
 		MeshComponent->SetHiddenInGame(true);
 	}
diff --git a/Engine/Plugins/Experimental/ColorCorrectRegions/Source/ColorCorrectRegions/Private/ColorCorrectWindow.cpp b/Engine/Plugins/Experimental/ColorCorrectRegions/Source/ColorCorrectRegions/Private/ColorCorrectWindow.cpp
--- a/Engine/Plugins/Experimental/ColorCorrectRegions/Source/ColorCorrectRegions/Private/ColorCorrectWindow.cpp
+++ b/Engine/Plugins/Experimental/ColorCorrectRegions/Source/ColorCorrectRegions/Private/ColorCorrectWindow.cpp
@@ -4,6 +4,7 @@
 #include "Components/BillboardComponent.h"
 #include "CoreMinimal.h"
 #include "Materials/Material.h"
+#include "Engine/CollisionProfile.h"
 #include "Engine/StaticMesh.h"
 #include "Components/StaticMeshComponent.h"
 #include "Engine/StaticMesh.h"
@@ -44,8 +45,8 @@ AColorCorrectionWindow::AColorCorrectionWindow(const FObjectInitializer& ObjectI
 		MeshComponent->SetStaticMesh(StaticMeshes[static_cast<uint8>(CCWType)]);
 		MeshComponent->SetWorldScale3D(FVector(1., 1., 0.001));
 		MeshComponent->SetMaterial(0, Material);
-		MeshComponent->SetCollisionEnabled(ECollisionEnabled::QueryOnly);
-		MeshComponent->SetCollisionProfileName(TEXT("OverlapAll"));
+		MeshComponent->SetCollisionEnabled(ECollisionEnabled::NoCollision);
+		MeshComponent->SetCollisionProfileName(UCollisionProfile::NoCollision_ProfileName);
 		MeshComponent->CastShadow = false;
 		MeshComponent->SetHiddenInGame(true);
 	}
diff --git a/Engine/Plugins/Experimental/ColorCorrectRegions/Source/ColorCorrectRegions/Public/ColorCorrectRegion.h b/Engine/Plugins/Experimental/ColorCorrectRegions/Source/ColorCorrectRegions/Public/ColorCorrectRegion.h
--- a/Engine/Plugins/Experimental/ColorCorrectRegions/Source/ColorCorrectRegions/Public/ColorCorrectRegion.h
+++ b/Engine/Plugins/Experimental/ColorCorrectRegions/Source/ColorCorrectRegions/Public/ColorCorrectRegion.h
@@ -230,8 +230,9 @@ private:
 
 	/**
 	* AffectedActors property change could potentially invoke a Dialog Window, which should be displayed on Game Thread.
+	* ActorListChangeType represents EPropertyChangeType
 	*/
-	void HandleAffectedActorsPropertyChange();
+	void HandleAffectedActorsPropertyChange(uint32 ActorListChangeType);
 
 	/**
 	* Copy state required for rendering to be consumed by Scene view extension.
@@ -244,9 +245,6 @@ private:
 	/** A copy of all properties required by render thread to process this CCR. */
 	FColorCorrectRenderProxyPtr ColorCorrectRenderProxy;
 
-	bool bActorListIsDirty;
-	uint8 ActorListChangeType;
-
 	FCriticalSection StateCopyCriticalSecion;
 
 	// This is for optimization purposes that would let us check assigned actors component's stencil ids ever few once in a while.
