Color correction Region/Windows: Added right click menu that allows users to assign any selected actor to Color Correction Actors Per actor CC.

diff --git a/Engine/Plugins/Experimental/ColorCorrectRegions/Source/ColorCorrectRegionsEditor/ColorCorrectRegionsEditor.Build.cs b/Engine/Plugins/Experimental/ColorCorrectRegions/Source/ColorCorrectRegionsEditor/ColorCorrectRegionsEditor.Build.cs
--- a/Engine/Plugins/Experimental/ColorCorrectRegions/Source/ColorCorrectRegionsEditor/ColorCorrectRegionsEditor.Build.cs
+++ b/Engine/Plugins/Experimental/ColorCorrectRegions/Source/ColorCorrectRegionsEditor/ColorCorrectRegionsEditor.Build.cs
@@ -6,21 +6,10 @@ public class ColorCorrectRegionsEditor : ModuleRules
 {
 	public ColorCorrectRegionsEditor(ReadOnlyTargetRules Target) : base(Target)
 	{
-		PublicDependencyModuleNames.AddRange(
-			new string[]
-			{
-				"Core",
-				"RenderCore",
-				"UnrealEd",
-				"ColorCorrectRegions",
-				"Slate",
-				"SlateCore",
-			}
-		);
-			
 		PrivateDependencyModuleNames.AddRange(
 			new string[]
 			{
+				"ActorPickerMode",
 				"Core",
 				"CoreUObject",
 				"Engine",
@@ -30,6 +19,7 @@ public class ColorCorrectRegionsEditor : ModuleRules
 				"SlateCore",
 				"ColorCorrectRegions",
 				"PlacementMode",
+				"SceneOutliner",
 			}
 		);
 	}
diff --git a/Engine/Plugins/Experimental/ColorCorrectRegions/Source/ColorCorrectRegionsEditor/Private/ColorCorrectRegionsEditorModule.cpp b/Engine/Plugins/Experimental/ColorCorrectRegions/Source/ColorCorrectRegionsEditor/Private/ColorCorrectRegionsEditorModule.cpp
--- a/Engine/Plugins/Experimental/ColorCorrectRegions/Source/ColorCorrectRegionsEditor/Private/ColorCorrectRegionsEditorModule.cpp
+++ b/Engine/Plugins/Experimental/ColorCorrectRegions/Source/ColorCorrectRegionsEditor/Private/ColorCorrectRegionsEditorModule.cpp
@@ -2,6 +2,7 @@
 
 #include "ColorCorrectRegionsEditorModule.h"
 #include "ActorFactories/ActorFactoryBlueprint.h"
+#include "ColorCorrectionActorContextMenu.h"
 #include "ColorCorrectRegion.h"
 #include "ColorCorrectRegionCustomization.h"
 #include "ColorCorrectRegionsStyle.h"
@@ -17,6 +18,9 @@ void FColorCorrectRegionsEditorModule::StartupModule()
 
 	FPropertyEditorModule& PropertyModule = FModuleManager::LoadModuleChecked<FPropertyEditorModule>("PropertyEditor");
 	PropertyModule.RegisterCustomClassLayout(AColorCorrectRegion::StaticClass()->GetFName(), FOnGetDetailCustomizationInstance::CreateStatic(&FColorCorrectWindowDetails::MakeInstance));
+
+	ContextMenu = MakeShared<FColorCorrectionActorContextMenu>();
+	ContextMenu->RegisterContextMenuExtender();
 }
 
 void FColorCorrectRegionsEditorModule::OnPlacementModeRefresh(FName CategoryName)
@@ -55,6 +59,7 @@ void FColorCorrectRegionsEditorModule::OnPlacementModeRefresh(FName CategoryName
 void FColorCorrectRegionsEditorModule::ShutdownModule()
 {
 	FColorCorrectRegionsStyle::Shutdown();
+	ContextMenu->UnregisterContextMenuExtender();
 }
 
 #undef LOCTEXT_NAMESPACE
diff --git a/Engine/Plugins/Experimental/ColorCorrectRegions/Source/ColorCorrectRegionsEditor/Private/ColorCorrectionActorContextMenu.cpp b/Engine/Plugins/Experimental/ColorCorrectRegions/Source/ColorCorrectRegionsEditor/Private/ColorCorrectionActorContextMenu.cpp
new file mode 100644
--- /dev/null
+++ b/Engine/Plugins/Experimental/ColorCorrectRegions/Source/ColorCorrectRegionsEditor/Private/ColorCorrectionActorContextMenu.cpp
@@ -0,0 +1,309 @@
+// Copyright Epic Games, Inc. All Rights Reserved.
+
+#include "ColorCorrectionActorContextMenu.h"
+#include "ActorPickerMode.h"
+#include "ActorTreeItem.h"
+#include "ColorCorrectRegion.h"
+#include "ColorCorrectWindow.h"
+#include "Kismet/GameplayStatics.h"
+#include "LevelEditor.h"
+#include "LevelEditor.h"
+#include "Editor/SceneOutliner/Public/ActorMode.h"
+#include "Engine/Selection.h"
+#include "Framework/Application/SlateApplication.h"
+#include "Editor/SceneOutliner/Public/SSceneOutliner.h"
+#include "Widgets/Input/SButton.h"
+#include "SSocketChooser.h"
+#include "Runtime/Engine/Classes/GameFramework/WorldSettings.h"
+#include "Editor/UnrealEdEngine.h"
+#include "UnrealEdGlobals.h"
+
+#define LOCTEXT_NAMESPACE "FColorCorrectRegionsContextMenu"
+
+namespace
+{
+	enum ECCType
+	{
+		Window,
+		Region
+	};
+
+	class FCCActorPickingMode : public FActorMode
+	{
+	public:
+		FCCActorPickingMode(SSceneOutliner* InSceneOutliner, FOnSceneOutlinerItemPicked InOnCCActorPicked);
+		virtual ~FCCActorPickingMode() {}
+
+		/* Begin ISceneOutlinerMode Implementation */
+		virtual void OnItemSelectionChanged(FSceneOutlinerTreeItemPtr Item, ESelectInfo::Type SelectionType, const FSceneOutlinerItemSelection& Selection) override;
+		virtual bool ShowViewButton() const override { return true; }
+		virtual bool ShouldShowFolders() const { return false; }
+		/* End ISceneOutlinerMode Implementation */
+	protected:
+
+		/** Callback for when CC Actor is selected. */
+		FOnSceneOutlinerItemPicked OnCCActorPicked;
+	};
+
+
+	/** Add Selected actors to the chosen CC Actor's Per actor CC. */
+	static void AddActorsToPerActorCC(AColorCorrectRegion* CCActorPtr, TSharedPtr<TArray<AActor*>> InSelectedActors)
+	{
+		// Hide all context menus.
+		FSlateApplication::Get().DismissAllMenus();
+
+		CCActorPtr->bEnablePerActorCC = true;
+		for (AActor* SelectedActor : *InSelectedActors)
+		{
+			TSoftObjectPtr<AActor> SelectedActorPtr(SelectedActor);
+
+			// Forcing property change event.
+			const FName AffectedActorsPropertyName = GET_MEMBER_NAME_CHECKED(AColorCorrectRegion, AffectedActors);
+			FProperty* AffectedActorsProperty = FindFieldChecked<FProperty>(AColorCorrectRegion::StaticClass(), AffectedActorsPropertyName);
+			CCActorPtr->PreEditChange(AffectedActorsProperty);
+
+			CCActorPtr->AffectedActors.Add(SelectedActorPtr);
+			FPropertyChangedEvent PropertyEvent(AffectedActorsProperty);
+			PropertyEvent.ChangeType = EPropertyChangeType::ArrayAdd;
+			CCActorPtr->PostEditChangeProperty(PropertyEvent);
+		}
+	}
+
+	/** Creates a new CCR or CCW and then adds Selected actors to Per Actor CC. */
+	static void CreateNewCCActor(ECCType InType, TSharedPtr<TArray<AActor*>> SelectedActors)
+	{
+		if (!SelectedActors.IsValid() || SelectedActors->Num() == 0)
+		{
+			return;
+		}
+
+		// Get bounds for the entire book 
+		FVector Origin;
+		FVector BoxExtent;
+		UGameplayStatics::GetActorArrayBounds(*SelectedActors, false /*bOnlyCollidingComponents*/, Origin, BoxExtent);
+		UWorld* World = (*SelectedActors)[0]->GetWorld();
+
+		if (!World)
+		{
+			return;
+		}
+
+		AWorldSettings* WorldSettings = World->GetWorldSettings();
+
+		if (!WorldSettings)
+		{
+			return;
+		}
+
+		TObjectPtr<AColorCorrectRegion> CCActorPtr;
+		FTransform Transform;
+		const FVector Scale = BoxExtent / (WorldSettings->WorldToMeters / 2.);
+
+		// Adding a 1% scale offset for a better ecompassing of selected actors.
+		Transform.SetScale3D(Scale*1.01);
+		Transform.SetLocation(Origin);
+
+		if (InType == ECCType::Window)
+		{
+			CCActorPtr = World->SpawnActor<AColorCorrectionWindow>();
+		}
+		else
+		{
+			AColorCorrectionRegion* CCActorRawPtr = World->SpawnActor<AColorCorrectionRegion>();
+			CCActorRawPtr->Type = EColorCorrectRegionsType::Box;
+			CCActorRawPtr->SetMeshVisibilityForRegionType();
+			CCActorPtr = CCActorRawPtr;
+		}
+
+		CCActorPtr->SetActorTransform(Transform);
+		AddActorsToPerActorCC(CCActorPtr, SelectedActors);
+
+		// Shift selection to the newly created CC Actor.
+		GEditor->SelectNone(false/*bNoteSelectionChange*/, true/*bDeselectBSPSurfs*/);
+		GEditor->SelectActor(CCActorPtr, true/*bInSelected*/, true/*bNotify*/);
+	};
+
+	/** Called by outliner when a tree item is selected by the user. */
+	static void OnOutlinerTreeItemSelected(TSharedRef<ISceneOutlinerTreeItem> NewParent, TSharedPtr<TArray<AActor*>> SelectedActors)
+	{
+		if (FActorTreeItem* ActorItem = NewParent->CastTo<FActorTreeItem>())
+		{
+			if (AColorCorrectRegion* CCR = Cast<AColorCorrectRegion>(ActorItem->Actor))
+			{
+				AddActorsToPerActorCC(CCR, SelectedActors);
+			}
+		}
+	};
+
+	/** A helper function used to determine if selected actor is either CCR or CCW. */
+	static bool IsActorCCR(const AActor* const InActor)
+	{
+		return (Cast<AColorCorrectRegion>(InActor) != nullptr);
+	}
+
+	/** Transfers Editor into a picker state when users select a pipet button. */
+	static FReply PickCCActorMode(TSharedPtr<TArray<AActor*>> InSelectedActors)
+	{
+		FSlateApplication::Get().DismissAllMenus();
+		FActorPickerModeModule& ActorPickerModeModule = FModuleManager::Get().GetModuleChecked<FActorPickerModeModule>("ActorPickerMode");
+
+		ActorPickerModeModule.BeginActorPickingMode(
+			FOnGetAllowedClasses(),
+			FOnShouldFilterActor::CreateStatic(IsActorCCR),
+			FOnActorSelected::CreateLambda([](AActor* InActor, TSharedPtr<TArray<AActor*>> SelectedActors)
+				{
+					AddActorsToPerActorCC(Cast<AColorCorrectRegion>(InActor), SelectedActors);
+				}, 
+				InSelectedActors)
+		);
+		return FReply::Handled();
+	}
+
+	/** Creates and populates menu allowing users to add selected actors to Per Actor CC. */
+	static void AddAttachToPerActorCCRMenu(class FMenuBuilder& MenuBuilder, TSharedPtr<TArray<AActor*>> InSelectedActors)
+	{
+		MenuBuilder.BeginSection("ColorCorrectionRegionsSection", LOCTEXT("ColorCorrectionRegions", "Color Correction Regions"));
+
+		auto CCActorAddToPerActorSubMenu = [](class FMenuBuilder& MenuBuilder, TSharedPtr<TArray<AActor*>> InSelectedActors)
+		{
+			FUIAction Action_CreateNewCCW(FExecuteAction::CreateStatic(CreateNewCCActor, Window, InSelectedActors));
+			FUIAction Action_CreateNewCCR(FExecuteAction::CreateStatic(CreateNewCCActor, Region, InSelectedActors));
+
+			MenuBuilder.AddMenuEntry(
+				LOCTEXT("MenuCreateAttachCCW", "Add to New Color Correction Window"),
+				LOCTEXT("MenuCreateAttachCCW_Tooltip", "Creates new Color Correction Window (CCW) and adds valid selected actors to Per-Actor CC of the newly created CCW."),
+				FSlateIcon(),
+				Action_CreateNewCCW,
+				NAME_None,
+				EUserInterfaceActionType::Button);
+
+			MenuBuilder.AddMenuEntry(
+				LOCTEXT("MenuCreateAttachCCR", "Add to New Color Correction Region"),
+				LOCTEXT("MenuCreateAttachCCR_Tooltip", "Creates new Color Correction Region (CCR) and adds valid selected actors to Per-Actor CC of the newly created CCR."),
+				FSlateIcon(),
+				Action_CreateNewCCR,
+				NAME_None,
+				EUserInterfaceActionType::Button);
+
+			MenuBuilder.AddMenuSeparator();
+
+			FCreateSceneOutlinerMode ModeFactory = FCreateSceneOutlinerMode::CreateLambda([](SSceneOutliner* Outliner, TSharedPtr<TArray<AActor*>> SelectedActors)
+				{
+					return new FCCActorPickingMode(Outliner, FOnSceneOutlinerItemPicked::CreateStatic(OnOutlinerTreeItemSelected, SelectedActors));
+				}, InSelectedActors);
+
+			FSceneOutlinerInitializationOptions CCSceneOutlinerInitOptions;
+			CCSceneOutlinerInitOptions.bShowHeaderRow = false;
+			CCSceneOutlinerInitOptions.bFocusSearchBoxWhenOpened = true;
+			CCSceneOutlinerInitOptions.ModeFactory = ModeFactory;
+			CCSceneOutlinerInitOptions.Filters->AddFilterPredicate<FActorTreeItem>(FActorTreeItem::FFilterPredicate::CreateStatic(IsActorCCR));
+
+			TSharedRef<SWidget> CCSceneOutliner =
+				SNew(SHorizontalBox)
+				+ SHorizontalBox::Slot()
+				[
+					SNew(SSceneOutliner, CCSceneOutlinerInitOptions)
+					.IsEnabled(FSlateApplication::Get().GetNormalExecutionAttribute())
+				]
+
+				+SHorizontalBox::Slot()
+				.VAlign(VAlign_Top)
+				.AutoWidth()
+				[
+					SNew(SVerticalBox)
+
+					+SVerticalBox::Slot()
+					.AutoHeight()
+					.Padding(0.0f, 9.0f, 0.0f, 0.0f)
+					[
+						SNew(SButton)
+						.ToolTipText( LOCTEXT( "PickCCActorButtonLabel", "Pick a Color Correction Actor") )
+						.ButtonStyle(FAppStyle::Get(), "HoverHintOnly")
+						.OnClicked(FOnClicked::CreateStatic(PickCCActorMode, InSelectedActors))
+						.ForegroundColor(FSlateColor::UseForeground())
+						.IsFocusable(false)
+						[
+							SNew(SImage)
+							.Image(FAppStyle::GetBrush("Icons.EyeDropper"))
+							.ColorAndOpacity(FSlateColor::UseForeground())
+						]
+					]
+				];
+			MenuBuilder.AddWidget(CCSceneOutliner, FText::GetEmpty(), /*bNoIndent=*/ true);
+		};
+
+		MenuBuilder.AddSubMenu(
+			LOCTEXT("CCRAttachToActorSubmenu", "Add to Per-Actor CC"),
+			LOCTEXT("CCRAttachToActorSubmenuToolTip", "Add currently selected actors to the list of actors that get affected by selected CC Actor in the following menu."),
+			FNewMenuDelegate::CreateLambda(CCActorAddToPerActorSubMenu, InSelectedActors),
+			false,
+			FSlateIcon()
+		);
+
+		MenuBuilder.EndSection();
+	}
+
+	/** Creates a menu extender that is responsible for adding buttons and outliners for Per Actor CC. */
+	static TSharedRef<FExtender> ExtendContextMenu(const TSharedRef<FUICommandList> CommandList, TArray<AActor*> InSelectedActors)
+	{
+		TSharedPtr<FExtender> Extender = MakeShared<FExtender>();
+
+		for (AActor* SelectedActor : InSelectedActors)
+		{
+			for (UActorComponent* Component : SelectedActor->GetComponents())
+			{
+				if (UPrimitiveComponent* PrimitiveComponent = Cast<UPrimitiveComponent>(Component))
+				{
+					TSharedPtr<TArray<AActor*>> SelectedActorsPtr = MakeShared<TArray<AActor*>>(MoveTemp(InSelectedActors));
+					Extender->AddMenuExtension("ActorTypeTools", EExtensionHook::After, CommandList,
+						FMenuExtensionDelegate::CreateStatic(AddAttachToPerActorCCRMenu, SelectedActorsPtr));
+					return Extender.ToSharedRef();
+				}
+			}
+		}
+		return Extender.ToSharedRef();
+	}
+}
+
+void FColorCorrectionActorContextMenu::RegisterContextMenuExtender()
+{
+	CCRLevelViewportContextMenuExtender = FLevelEditorModule::FLevelViewportMenuExtender_SelectedActors::CreateStatic(ExtendContextMenu);
+	FLevelEditorModule& LevelEditorModule = FModuleManager::Get().LoadModuleChecked<FLevelEditorModule>("LevelEditor");
+	TArray<FLevelEditorModule::FLevelViewportMenuExtender_SelectedActors>& MenuExtenders = LevelEditorModule.GetAllLevelViewportContextMenuExtenders();
+	MenuExtenders.Add(CCRLevelViewportContextMenuExtender);
+	ContextMenuExtenderDelegateHandle = MenuExtenders.Last().GetHandle();
+}
+
+void FColorCorrectionActorContextMenu::UnregisterContextMenuExtender()
+{
+	if (FModuleManager::Get().IsModuleLoaded("LevelEditor"))
+	{
+		FLevelEditorModule& LevelEditorModule = FModuleManager::LoadModuleChecked<FLevelEditorModule>("LevelEditor");
+		LevelEditorModule.GetAllLevelViewportContextMenuExtenders().RemoveAll(
+			[&](const FLevelEditorModule::FLevelViewportMenuExtender_SelectedActors& Delegate)
+			{
+				return Delegate.GetHandle() == ContextMenuExtenderDelegateHandle;
+			});
+	}
+}
+
+FCCActorPickingMode::FCCActorPickingMode(SSceneOutliner* InSceneOutliner, FOnSceneOutlinerItemPicked InOnCCActorPicked)
+	: FActorMode(FActorModeParams(InSceneOutliner))
+	, OnCCActorPicked(InOnCCActorPicked)
+{
+}
+
+void FCCActorPickingMode::OnItemSelectionChanged(FSceneOutlinerTreeItemPtr Item, ESelectInfo::Type SelectionType, const FSceneOutlinerItemSelection& Selection)
+{
+	auto SelectedItems = SceneOutliner->GetSelectedItems();
+	if (SelectedItems.Num() > 0)
+	{
+		auto FirstItem = SelectedItems[0];
+		if (FirstItem->CanInteract())
+		{
+			OnCCActorPicked.ExecuteIfBound(FirstItem.ToSharedRef());
+		}
+	}
+}
+
+#undef LOCTEXT_NAMESPACE
diff --git a/Engine/Plugins/Experimental/ColorCorrectRegions/Source/ColorCorrectRegionsEditor/Private/ColorCorrectionActorContextMenu.h b/Engine/Plugins/Experimental/ColorCorrectRegions/Source/ColorCorrectRegionsEditor/Private/ColorCorrectionActorContextMenu.h
new file mode 100644
--- /dev/null
+++ b/Engine/Plugins/Experimental/ColorCorrectRegions/Source/ColorCorrectRegionsEditor/Private/ColorCorrectionActorContextMenu.h
@@ -0,0 +1,23 @@
+// Copyright Epic Games, Inc. All Rights Reserved.
+
+#pragma once
+
+#include "ActorMode.h"
+#include "LevelEditor.h"
+
+class FColorCorrectionActorContextMenu
+{
+public:
+	/** Create all required callbacks for context menu's widget and button creation. */
+	void RegisterContextMenuExtender();
+
+	/** Release context menu related resources such as handles. */ 
+	void UnregisterContextMenuExtender();
+
+private:
+	/** Context menu extender. */
+	FLevelEditorModule::FLevelViewportMenuExtender_SelectedActors CCRLevelViewportContextMenuExtender;
+
+	/** Handle to context menu extender delegate. */
+	FDelegateHandle ContextMenuExtenderDelegateHandle;
+};
diff --git a/Engine/Plugins/Experimental/ColorCorrectRegions/Source/ColorCorrectRegionsEditor/Public/ColorCorrectRegionsEditorModule.h b/Engine/Plugins/Experimental/ColorCorrectRegions/Source/ColorCorrectRegionsEditor/Public/ColorCorrectRegionsEditorModule.h
--- a/Engine/Plugins/Experimental/ColorCorrectRegions/Source/ColorCorrectRegionsEditor/Public/ColorCorrectRegionsEditorModule.h
+++ b/Engine/Plugins/Experimental/ColorCorrectRegions/Source/ColorCorrectRegionsEditor/Public/ColorCorrectRegionsEditorModule.h
@@ -17,5 +17,7 @@ public:
 	/** IModuleInterface implementation end */
 
 	void OnPlacementModeRefresh(FName CategoryName);
-	void OnPostEngineInit();
+
+private:
+	TSharedPtr<class FColorCorrectionActorContextMenu> ContextMenu;
 };
