Level Editor / Virtual Production Plugins:

-Fix: Virtual Production Actors Missing from Default Outliner Filters

diff --git a/Engine/Plugins/Experimental/ColorCorrectRegions/Source/ColorCorrectRegions/ColorCorrectRegions.Build.cs b/Engine/Plugins/Experimental/ColorCorrectRegions/Source/ColorCorrectRegions/ColorCorrectRegions.Build.cs
--- a/Engine/Plugins/Experimental/ColorCorrectRegions/Source/ColorCorrectRegions/ColorCorrectRegions.Build.cs
+++ b/Engine/Plugins/Experimental/ColorCorrectRegions/Source/ColorCorrectRegions/ColorCorrectRegions.Build.cs
@@ -36,7 +36,7 @@ public class ColorCorrectRegions : ModuleRules
 				"SlateCore",
 				"Renderer",
 				"Projects",
-				"RenderCore",
+				"RenderCore"
 			}
 		);
 
@@ -44,6 +44,7 @@ public class ColorCorrectRegions : ModuleRules
 		{
 			PublicDependencyModuleNames.Add("UnrealEd");
 			PrivateDependencyModuleNames.Add("UnrealEd");
+			PrivateDependencyModuleNames.Add("EditorWidgets");
 		}
 	}
 }
diff --git a/Engine/Plugins/Experimental/ColorCorrectRegions/Source/ColorCorrectRegions/Private/ColorCorrectRegionsModule.cpp b/Engine/Plugins/Experimental/ColorCorrectRegions/Source/ColorCorrectRegions/Private/ColorCorrectRegionsModule.cpp
--- a/Engine/Plugins/Experimental/ColorCorrectRegions/Source/ColorCorrectRegions/Private/ColorCorrectRegionsModule.cpp
+++ b/Engine/Plugins/Experimental/ColorCorrectRegions/Source/ColorCorrectRegions/Private/ColorCorrectRegionsModule.cpp
@@ -1,18 +1,27 @@
 // Copyright Epic Games, Inc. All Rights Reserved.
 
 #include "ColorCorrectRegionsModule.h"
+
+#include "ColorCorrectRegion.h"
+#include "ColorCorrectWindow.h"
+#if WITH_EDITOR
+#include "Filters/CustomClassFilterData.h"
+#endif
 #include "Interfaces/IPluginManager.h"
+#include "LevelEditor.h"
 #include "Misc/Paths.h"
 #include "ShaderCore.h"
 
-
-
 #define LOCTEXT_NAMESPACE "FColorCorrectRegionsModule"
 
 void FColorCorrectRegionsModule::StartupModule()
 {
 	FString PluginShaderDir = FPaths::Combine(IPluginManager::Get().FindPlugin(TEXT("ColorCorrectRegions"))->GetBaseDir(), TEXT("Shaders"));
 	AddShaderSourceDirectoryMapping(TEXT("/ColorCorrectRegionsShaders"), PluginShaderDir);
+
+#if WITH_EDITOR
+	RegisterOutlinerFilters();
+#endif
 }
 
 void FColorCorrectRegionsModule::ShutdownModule()
@@ -20,6 +29,25 @@ void FColorCorrectRegionsModule::ShutdownModule()
 
 }
 
+#if WITH_EDITOR
+void FColorCorrectRegionsModule::RegisterOutlinerFilters()
+{
+	if (FLevelEditorModule* LevelEditorModule = FModuleManager::GetModulePtr<FLevelEditorModule>(TEXT("LevelEditor")))
+	{
+		if (const TSharedPtr<FFilterCategory> FilterCategory = LevelEditorModule->GetOutlinerFilterCategory(FLevelEditorOutlinerBuiltInCategories::VirtualProduction()))
+		{
+			const TSharedRef<FCustomClassFilterData> ColorCorrectionRegionActorClassData =
+				MakeShared<FCustomClassFilterData>(AColorCorrectionRegion::StaticClass(), FilterCategory, FLinearColor::White);
+			LevelEditorModule->AddCustomClassFilterToOutliner(ColorCorrectionRegionActorClassData);
+
+			const TSharedRef<FCustomClassFilterData> ColorCorrectionWindowActorClassData =
+				MakeShared<FCustomClassFilterData>(AColorCorrectionWindow::StaticClass(), FilterCategory, FLinearColor::White);
+			LevelEditorModule->AddCustomClassFilterToOutliner(ColorCorrectionWindowActorClassData);
+		}
+	}
+}
+#endif
+
 #undef LOCTEXT_NAMESPACE
 	
 IMPLEMENT_MODULE(FColorCorrectRegionsModule, ColorCorrectRegions);
diff --git a/Engine/Plugins/Experimental/ColorCorrectRegions/Source/ColorCorrectRegions/Public/ColorCorrectRegionsModule.h b/Engine/Plugins/Experimental/ColorCorrectRegions/Source/ColorCorrectRegions/Public/ColorCorrectRegionsModule.h
--- a/Engine/Plugins/Experimental/ColorCorrectRegions/Source/ColorCorrectRegions/Public/ColorCorrectRegionsModule.h
+++ b/Engine/Plugins/Experimental/ColorCorrectRegions/Source/ColorCorrectRegions/Public/ColorCorrectRegionsModule.h
@@ -15,4 +15,8 @@ public:
 	virtual void StartupModule() override;
 	virtual void ShutdownModule() override;
 	/** IModuleInterface implementation end */
+
+#if WITH_EDITOR
+	void RegisterOutlinerFilters();
+#endif
 };
diff --git a/Engine/Plugins/Experimental/VirtualCamera/Source/VirtualCameraEditor/Private/VirtualCameraEditorModule.cpp b/Engine/Plugins/Experimental/VirtualCamera/Source/VirtualCameraEditor/Private/VirtualCameraEditorModule.cpp
--- a/Engine/Plugins/Experimental/VirtualCamera/Source/VirtualCameraEditor/Private/VirtualCameraEditorModule.cpp
+++ b/Engine/Plugins/Experimental/VirtualCamera/Source/VirtualCameraEditor/Private/VirtualCameraEditorModule.cpp
@@ -15,6 +15,7 @@
 #include "IVPUtilitiesEditorModule.h"
 #include "LevelEditor.h"
 #include "LevelEditorOutlinerSettings.h"
+#include "VirtualCameraActor.h"
 
 #define LOCTEXT_NAMESPACE "FVirtualCameraEditorModule"
 
@@ -71,14 +72,21 @@ private:
 
 	void RegisterOutlinerFilters()
 	{
-		FLevelEditorModule& LevelEditorModule = FModuleManager::LoadModuleChecked<FLevelEditorModule>(TEXT("LevelEditor"));
-		if (const TSharedPtr<FFilterCategory> VPFilterCategory = LevelEditorModule.GetOutlinerFilterCategory(FLevelEditorOutlinerBuiltInCategories::VirtualProduction()))
+		if (FLevelEditorModule* LevelEditorModule = FModuleManager::GetModulePtr<FLevelEditorModule>(TEXT("LevelEditor")))
 		{
-			TSharedRef<FCustomClassFilterData> CineCameraActorClassData = MakeShared<FCustomClassFilterData>(ACineCameraActor::StaticClass(), VPFilterCategory, FLinearColor::White);
-			LevelEditorModule.AddCustomClassFilterToOutliner(CineCameraActorClassData);
+			if (const TSharedPtr<FFilterCategory> VPFilterCategory = LevelEditorModule->GetOutlinerFilterCategory(FLevelEditorOutlinerBuiltInCategories::VirtualProduction()))
+			{
+				TSharedRef<FCustomClassFilterData> CineCameraActorClassData =
+					MakeShared<FCustomClassFilterData>(ACineCameraActor::StaticClass(), VPFilterCategory, FLinearColor::White);
+				LevelEditorModule->AddCustomClassFilterToOutliner(CineCameraActorClassData);
+
+				TSharedRef<FCustomClassFilterData> VirtualCameraActorClassData =
+					MakeShared<FCustomClassFilterData>(AVirtualCameraActor::StaticClass(), VPFilterCategory, FLinearColor::White);
+				LevelEditorModule->AddCustomClassFilterToOutliner(VirtualCameraActorClassData);
+			}
 		}
 	}
-
+	
 	void RegisterSettings()
 	{
 		ISettingsModule* SettingsModule = FModuleManager::GetModulePtr<ISettingsModule>("Settings");
diff --git a/Engine/Plugins/Runtime/nDisplay/Source/DisplayCluster/DisplayCluster.Build.cs b/Engine/Plugins/Runtime/nDisplay/Source/DisplayCluster/DisplayCluster.Build.cs
--- a/Engine/Plugins/Runtime/nDisplay/Source/DisplayCluster/DisplayCluster.Build.cs
+++ b/Engine/Plugins/Runtime/nDisplay/Source/DisplayCluster/DisplayCluster.Build.cs
@@ -52,6 +52,7 @@ public class DisplayCluster : ModuleRules
 			PublicIncludePathModuleNames.Add("DisplayClusterConfigurator");
 
 			PrivateDependencyModuleNames.Add("UnrealEd");
+			PrivateDependencyModuleNames.Add("EditorWidgets");
 			PrivateDependencyModuleNames.Add("LevelEditor");
 		}
 
diff --git a/Engine/Plugins/Runtime/nDisplay/Source/DisplayCluster/Private/DisplayClusterModule.cpp b/Engine/Plugins/Runtime/nDisplay/Source/DisplayCluster/Private/DisplayClusterModule.cpp
--- a/Engine/Plugins/Runtime/nDisplay/Source/DisplayCluster/Private/DisplayClusterModule.cpp
+++ b/Engine/Plugins/Runtime/nDisplay/Source/DisplayCluster/Private/DisplayClusterModule.cpp
@@ -8,6 +8,12 @@
 #include "Render/DisplayClusterRenderManager.h"
 
 #include "DisplayClusterConfigurationTypes.h"
+#include "DisplayClusterLightCardActor.h"
+#include "DisplayClusterRootActor.h"
+#if WITH_EDITOR
+#include "Filters/CustomClassFilterData.h"
+#endif
+#include "LevelEditor.h"
 
 #include "Misc/DisplayClusterGlobals.h"
 #include "Misc/DisplayClusterLog.h"
@@ -39,6 +45,10 @@ FDisplayClusterModule::~FDisplayClusterModule()
 void FDisplayClusterModule::StartupModule()
 {
 	UE_LOG(LogDisplayClusterModule, Log, TEXT("DisplayCluster module has been started"));
+
+#if WITH_EDITOR
+	RegisterOutlinerFilters();
+#endif
 }
 
 void FDisplayClusterModule::ShutdownModule()
@@ -47,6 +57,24 @@ void FDisplayClusterModule::ShutdownModule()
 	Release();
 }
 
+#if WITH_EDITOR
+void FDisplayClusterModule::RegisterOutlinerFilters()
+{
+	if (FLevelEditorModule* LevelEditorModule = FModuleManager::GetModulePtr<FLevelEditorModule>(TEXT("LevelEditor")))
+	{
+		if (const TSharedPtr<FFilterCategory> FilterCategory = LevelEditorModule->GetOutlinerFilterCategory(FLevelEditorOutlinerBuiltInCategories::VirtualProduction()))
+		{
+			const TSharedRef<FCustomClassFilterData> LightCardActorClassData =
+				MakeShared<FCustomClassFilterData>(ADisplayClusterLightCardActor::StaticClass(), FilterCategory, FLinearColor::White);
+			LevelEditorModule->AddCustomClassFilterToOutliner(LightCardActorClassData);
+
+			const TSharedRef<FCustomClassFilterData> RootActorClassData =
+				MakeShared<FCustomClassFilterData>(ADisplayClusterRootActor::StaticClass(), FilterCategory, FLinearColor::White);
+			LevelEditorModule->AddCustomClassFilterToOutliner(RootActorClassData);
+		}
+	}
+}
+#endif
 
 //////////////////////////////////////////////////////////////////////////////////////////////
 // IPDisplayCluster
diff --git a/Engine/Plugins/Runtime/nDisplay/Source/DisplayCluster/Private/DisplayClusterModule.h b/Engine/Plugins/Runtime/nDisplay/Source/DisplayCluster/Private/DisplayClusterModule.h
--- a/Engine/Plugins/Runtime/nDisplay/Source/DisplayCluster/Private/DisplayClusterModule.h
+++ b/Engine/Plugins/Runtime/nDisplay/Source/DisplayCluster/Private/DisplayClusterModule.h
@@ -78,6 +78,10 @@ private:
 	//////////////////////////////////////////////////////////////////////////////////////////////
 	virtual void StartupModule() override;
 	virtual void ShutdownModule() override;
+	
+#if WITH_EDITOR
+	void RegisterOutlinerFilters();
+#endif
 
 private:
 	// Is module initialized.
