Color Correction Actors: Fixed shader compiler issue for mobile renderers.

diff --git a/Engine/Plugins/Experimental/ColorCorrectRegions/Shaders/Private/ColorCorrectRegionsCommon.usf b/Engine/Plugins/Experimental/ColorCorrectRegions/Shaders/Private/ColorCorrectRegionsCommon.usf
--- a/Engine/Plugins/Experimental/ColorCorrectRegions/Shaders/Private/ColorCorrectRegionsCommon.usf
+++ b/Engine/Plugins/Experimental/ColorCorrectRegions/Shaders/Private/ColorCorrectRegionsCommon.usf
@@ -291,13 +291,29 @@ MaterialFloat3 RotateTranslateScale(MaterialFloat3 OriginalVec, MaterialFloat3 T
 /** Applies an offset to the scene texture lookup and decodes the HDR linear space color. */
 float4 SceneTextureLookup(float2 UV, int SceneTextureIndex, bool bFiltered)
 {
+#if SHADING_PATH_DEFERRED
 	FScreenSpaceData ScreenSpaceData = GetScreenSpaceData(UV, false);
+#endif
 	switch (SceneTextureIndex)
 	{
 		// order needs to match to ESceneTextureId
 
 	case PPI_SceneColor:
 		return float4(CalcSceneColor(UV), 0);
+	case PPI_SeparateTranslucency:
+		return float4(1, 1, 1, 1);
+
+	case PPI_PostProcessInput0:
+		return Texture2DSample(PostProcessInput_0_Texture, bFiltered ? PostProcessInput_BilinearSampler : PostProcessInput_0_SharedSampler, UV);
+	case PPI_PostProcessInput1:
+		return Texture2DSample(PostProcessInput_1_Texture, bFiltered ? PostProcessInput_BilinearSampler : PostProcessInput_1_SharedSampler, UV);
+	case PPI_PostProcessInput2:
+		return Texture2DSample(PostProcessInput_2_Texture, bFiltered ? PostProcessInput_BilinearSampler : PostProcessInput_2_SharedSampler, UV);
+	case PPI_PostProcessInput3:
+		return Texture2DSample(PostProcessInput_3_Texture, bFiltered ? PostProcessInput_BilinearSampler : PostProcessInput_3_SharedSampler, UV);
+	case PPI_PostProcessInput4:
+		return Texture2DSample(PostProcessInput_4_Texture, bFiltered ? PostProcessInput_BilinearSampler : PostProcessInput_4_SharedSampler, UV);
+#if SHADING_PATH_DEFERRED
 	case PPI_SceneDepth:
 		return ScreenSpaceData.GBuffer.Depth;
 	case PPI_DiffuseColor:
@@ -314,8 +330,6 @@ float4 SceneTextureLookup(float2 UV, int SceneTextureIndex, bool bFiltered)
 		return ScreenSpaceData.GBuffer.Metallic;
 	case PPI_WorldNormal:
 		return float4(ScreenSpaceData.GBuffer.WorldNormal, 0);
-	case PPI_SeparateTranslucency:
-		return float4(1, 1, 1, 1);
 	case PPI_Opacity:
 		return ScreenSpaceData.GBuffer.CustomData.a;
 	case PPI_Roughness:
@@ -324,20 +338,6 @@ float4 SceneTextureLookup(float2 UV, int SceneTextureIndex, bool bFiltered)
 		return ScreenSpaceData.GBuffer.GBufferAO;
 	case PPI_CustomDepth:
 		return ScreenSpaceData.GBuffer.CustomDepth;
-
-	case PPI_PostProcessInput0:
-		return Texture2DSample(PostProcessInput_0_Texture, bFiltered ? PostProcessInput_BilinearSampler : PostProcessInput_0_SharedSampler, UV);
-	case PPI_PostProcessInput1:
-		return Texture2DSample(PostProcessInput_1_Texture, bFiltered ? PostProcessInput_BilinearSampler : PostProcessInput_1_SharedSampler, UV);
-	case PPI_PostProcessInput2:
-		return Texture2DSample(PostProcessInput_2_Texture, bFiltered ? PostProcessInput_BilinearSampler : PostProcessInput_2_SharedSampler, UV);
-	case PPI_PostProcessInput3:
-		return Texture2DSample(PostProcessInput_3_Texture, bFiltered ? PostProcessInput_BilinearSampler : PostProcessInput_3_SharedSampler, UV);
-	case PPI_PostProcessInput4:
-		return Texture2DSample(PostProcessInput_4_Texture, bFiltered ? PostProcessInput_BilinearSampler : PostProcessInput_4_SharedSampler, UV);
-
-	case PPI_DecalMask:
-		return 0;  // material compiler will return an error
 	case PPI_ShadingModelColor:
 		return float4(GetShadingModelColor(ScreenSpaceData.GBuffer.ShadingModelID), 1);
 	case PPI_ShadingModelID:
@@ -355,6 +355,7 @@ float4 SceneTextureLookup(float2 UV, int SceneTextureIndex, bool bFiltered)
 		return float4(ScreenSpaceData.GBuffer.WorldTangent, 0);
 	case PPI_Anisotropy:
 		return ScreenSpaceData.GBuffer.Anisotropy;
+#endif
 	default:
 		return float4(0, 0, 0, 0);
 	}
diff --git a/Engine/Plugins/Experimental/ColorCorrectRegions/Shaders/Private/ColorCorrectRegionsStencilMerger.usf b/Engine/Plugins/Experimental/ColorCorrectRegions/Shaders/Private/ColorCorrectRegionsStencilMerger.usf
--- a/Engine/Plugins/Experimental/ColorCorrectRegions/Shaders/Private/ColorCorrectRegionsStencilMerger.usf
+++ b/Engine/Plugins/Experimental/ColorCorrectRegions/Shaders/Private/ColorCorrectRegionsStencilMerger.usf
@@ -16,6 +16,7 @@ uint MainPS(
 	noperspective float4 UVAndScreenPos : TEXCOORD0
 ) : SV_Target0
 {
+#if SHADING_PATH_DEFERRED
 	float4 SvPosition = UVAndScreenPos;
 	const float Float_Eps = 0.00000000001;
 
@@ -33,5 +34,6 @@ uint MainPS(
 			}
 		}
 	}
+#endif
 	return 0;
 }
diff --git a/Engine/Plugins/Experimental/ColorCorrectRegions/Source/ColorCorrectRegions/ColorCorrectRegions.Build.cs b/Engine/Plugins/Experimental/ColorCorrectRegions/Source/ColorCorrectRegions/ColorCorrectRegions.Build.cs
--- a/Engine/Plugins/Experimental/ColorCorrectRegions/Source/ColorCorrectRegions/ColorCorrectRegions.Build.cs
+++ b/Engine/Plugins/Experimental/ColorCorrectRegions/Source/ColorCorrectRegions/ColorCorrectRegions.Build.cs
@@ -17,8 +17,6 @@ public class ColorCorrectRegions : ModuleRules
 		PublicDependencyModuleNames.AddRange(
 			new string[]
 			{
-				"Core",
-				"RenderCore",
 				"DisplayClusterLightCardExtender"
 			}
 		);
@@ -29,20 +27,15 @@ public class ColorCorrectRegions : ModuleRules
 				"Core",
 				"CoreUObject",
 				"Engine",
-				"MovieSceneCapture",
-				"RenderCore",
 				"RHI",
-				"Slate",
-				"SlateCore",
 				"Renderer",
 				"Projects",
-				"RenderCore"
+				"RenderCore",
 			}
 		);
 
 		if (Target.Type == TargetType.Editor)
 		{
-			PublicDependencyModuleNames.Add("UnrealEd");
 			PrivateDependencyModuleNames.Add("UnrealEd");
 			PrivateDependencyModuleNames.Add("EditorWidgets");
 		}
