nDisplay: Support CCRs in the ICVFX panel. CCRs can be spawned in and recognized in the panel. Transforms still behave similar to CCWs.

diff --git a/Engine/Plugins/Experimental/ColorCorrectRegions/Source/ColorCorrectRegions/Private/ColorCorrectRegion.cpp b/Engine/Plugins/Experimental/ColorCorrectRegions/Source/ColorCorrectRegions/Private/ColorCorrectRegion.cpp
--- a/Engine/Plugins/Experimental/ColorCorrectRegions/Source/ColorCorrectRegions/Private/ColorCorrectRegion.cpp
+++ b/Engine/Plugins/Experimental/ColorCorrectRegions/Source/ColorCorrectRegions/Private/ColorCorrectRegion.cpp
@@ -16,6 +16,10 @@
 #include "Materials/Material.h"
 #include "UObject/ConstructorHelpers.h"
 
+#if WITH_EDITOR
+#include "IDisplayClusterLightCardExtenderModule.h"
+#endif
+
 ENUM_RANGE_BY_COUNT(EColorCorrectRegionsType, EColorCorrectRegionsType::MAX)
 
 
@@ -54,6 +58,14 @@ AColorCorrectRegion::AColorCorrectRegion(const FObjectInitializer& ObjectInitial
 		CreateIcon();
 	}
 #endif
+
+#if WITH_EDITOR
+	if (!IsTemplate())
+	{
+		IDisplayClusterLightCardExtenderModule& LightCardExtenderModule = IDisplayClusterLightCardExtenderModule::Get();
+		LightCardExtenderModule.GetOnSequencerTimeChanged().AddUObject(this, &AColorCorrectRegion::OnSequencerTimeChanged);
+	}
+#endif
 }
 
 void AColorCorrectRegion::BeginPlay()
@@ -212,6 +224,13 @@ void AColorCorrectRegion::TransferState()
 
 }
 
+#if WITH_EDITOR
+void AColorCorrectRegion::OnSequencerTimeChanged(TWeakPtr<ISequencer> InSequencer)
+{
+	UpdatePositionalParamsFromTransform();
+}
+#endif
+
 void AColorCorrectRegion::HandleAffectedActorsPropertyChange(uint32 ActorListChangeType)
 {
 	if (ActorListChangeType == EPropertyChangeType::ArrayAdd
@@ -298,10 +317,20 @@ void AColorCorrectRegion::CreateIcon()
 }
 #endif 
 
+AColorCorrectRegion::~AColorCorrectRegion()
+{
+#if WITH_EDITOR
+	if (!IsTemplate())
+	{
+		IDisplayClusterLightCardExtenderModule& LightCardExtenderModule = IDisplayClusterLightCardExtenderModule::Get();
+		LightCardExtenderModule.GetOnSequencerTimeChanged().RemoveAll(this);
+	}
+#endif
+}
+
 #if WITH_EDITOR
 void AColorCorrectRegion::PostEditChangeProperty(struct FPropertyChangedEvent& PropertyChangedEvent)
 {
-	Super::PostEditChangeProperty(PropertyChangedEvent);
 	const FName PropertyName = PropertyChangedEvent.GetPropertyName();
 
 	if (!ColorCorrectRegionsSubsystem.IsValid())
@@ -338,9 +367,184 @@ void AColorCorrectRegion::PostEditChangeProperty(struct FPropertyChangedEvent& P
 			ColorCorrectRegionsSubsystem->OnLevelsChanged();
 		}
 	}
+
+	// Stage actor properties
+	{
+		const FStructProperty* StructProperty = CastField<FStructProperty>(PropertyChangedEvent.MemberProperty);
+		const bool bIsOrientation = StructProperty ? StructProperty->Struct == FDisplayClusterPositionalParams::StaticStruct() : false;
+	
+		if (bIsOrientation)
+		{
+			UpdateStageActorTransform();
+			// Updates MU in real-time. Skip our method as the positional coordinates are already correct.
+			AActor::PostEditMove(PropertyChangedEvent.ChangeType != EPropertyChangeType::Interactive);
+		}
+		else if (
+			PropertyName == USceneComponent::GetRelativeLocationPropertyName() ||
+			PropertyName == USceneComponent::GetRelativeRotationPropertyName() ||
+			PropertyName == USceneComponent::GetRelativeScale3DPropertyName())
+		{
+			bNotifyOnParamSetter = false;
+			UpdatePositionalParamsFromTransform();
+			bNotifyOnParamSetter = true;
+		}
+	}
+
+	// Call after stage actor transform is updated, so any observers will have both the correct actor transform and
+	// positional properties.
+	Super::PostEditChangeProperty(PropertyChangedEvent);
+}
+
+void AColorCorrectRegion::PostEditMove(bool bFinished)
+{
+	Super::PostEditMove(bFinished);
+
+	bNotifyOnParamSetter = false;
+	UpdatePositionalParamsFromTransform();
+	bNotifyOnParamSetter = true;
 }
 #endif //WITH_EDITOR
 
+#define NOTIFY_PARAM_SETTER()\
+	if (bNotifyOnParamSetter)\
+	{\
+		UpdateStageActorTransform();\
+	}\
+
+void AColorCorrectRegion::SetLongitude(double InValue)
+{
+	PositionalParams.Longitude = InValue;
+	NOTIFY_PARAM_SETTER()
+}
+
+double AColorCorrectRegion::GetLongitude() const
+{
+	return PositionalParams.Longitude;
+}
+
+void AColorCorrectRegion::SetLatitude(double InValue)
+{
+	PositionalParams.Latitude = InValue;
+	NOTIFY_PARAM_SETTER()
+}
+
+double AColorCorrectRegion::GetLatitude() const
+{
+	return PositionalParams.Latitude;
+}
+
+void AColorCorrectRegion::SetDistanceFromCenter(double InValue)
+{
+	PositionalParams.DistanceFromCenter = InValue;
+	NOTIFY_PARAM_SETTER()
+}
+
+double AColorCorrectRegion::GetDistanceFromCenter() const
+{
+	return PositionalParams.DistanceFromCenter;
+}
+
+void AColorCorrectRegion::SetSpin(double InValue)
+{
+	PositionalParams.Spin = InValue;
+	NOTIFY_PARAM_SETTER()
+}
+
+double AColorCorrectRegion::GetSpin() const
+{
+	return PositionalParams.Spin;
+}
+
+void AColorCorrectRegion::SetPitch(double InValue)
+{
+	PositionalParams.Pitch = InValue;
+	NOTIFY_PARAM_SETTER()
+}
+
+double AColorCorrectRegion::GetPitch() const
+{
+	return PositionalParams.Pitch;
+}
+
+void AColorCorrectRegion::SetYaw(double InValue)
+{
+	PositionalParams.Yaw = InValue;
+	NOTIFY_PARAM_SETTER()
+}
+
+double AColorCorrectRegion::GetYaw() const
+{
+	return PositionalParams.Yaw;
+}
+
+void AColorCorrectRegion::SetRadialOffset(double InValue)
+{
+	PositionalParams.RadialOffset = InValue;
+	NOTIFY_PARAM_SETTER()
+}
+
+double AColorCorrectRegion::GetRadialOffset() const
+{
+	return PositionalParams.RadialOffset;
+}
+
+void AColorCorrectRegion::SetScale(const FVector2D& InScale)
+{
+	PositionalParams.Scale = InScale;
+	NOTIFY_PARAM_SETTER()
+}
+
+FVector2D AColorCorrectRegion::GetScale() const
+{
+	return PositionalParams.Scale;
+}
+
+void AColorCorrectRegion::SetOrigin(const FTransform& InOrigin)
+{
+	Origin = InOrigin;
+}
+
+FTransform AColorCorrectRegion::GetOrigin() const
+{
+	return Origin;
+}
+
+void AColorCorrectRegion::SetPositionalParams(const FDisplayClusterPositionalParams& InParams)
+{
+	PositionalParams = InParams;
+	NOTIFY_PARAM_SETTER()
+}
+
+FDisplayClusterPositionalParams AColorCorrectRegion::GetPositionalParams() const
+{
+	return PositionalParams;
+}
+
+void AColorCorrectRegion::GetPositionalProperties(FPositionalPropertyArray& OutPropertyPairs) const
+{
+	void* Container = (void*)(&PositionalParams);
+
+	const TSet<FName>& PropertyNames = GetPositionalPropertyNames();
+	OutPropertyPairs.Reserve(PropertyNames.Num());
+
+	for (const FName& PropertyName : PropertyNames)
+	{
+		if (FProperty* Property = FindFProperty<FProperty>(FDisplayClusterPositionalParams::StaticStruct(), PropertyName))
+		{
+			OutPropertyPairs.Emplace(Container, Property);
+		}
+	}
+
+	if (FStructProperty* ParamsProperty = FindFProperty<FStructProperty>(GetClass(), GET_MEMBER_NAME_CHECKED(AColorCorrectRegion, PositionalParams)))
+	{
+		OutPropertyPairs.Emplace((void*)this, ParamsProperty);
+	}
+}
+
+FName AColorCorrectRegion::GetPositionalPropertiesMemberName() const
+{
+	return GET_MEMBER_NAME_CHECKED(AColorCorrectRegion, PositionalParams);
+}
 
 AColorCorrectionRegion::AColorCorrectionRegion(const FObjectInitializer& ObjectInitializer)
 	: Super(ObjectInitializer)
@@ -399,8 +603,11 @@ void AColorCorrectionRegion::PostEditChangeProperty(struct FPropertyChangedEvent
 		SetMeshVisibilityForRegionType();
 	}
 }
+
 FName AColorCorrectionRegion::GetCustomIconName() const
 {
 	return TEXT("CCR.OutlinerThumbnail");
 }
 #endif
+
+#undef NOTIFY_PARAM_SETTER
diff --git a/Engine/Plugins/Experimental/ColorCorrectRegions/Source/ColorCorrectRegions/Private/ColorCorrectWindow.cpp b/Engine/Plugins/Experimental/ColorCorrectRegions/Source/ColorCorrectRegions/Private/ColorCorrectWindow.cpp
--- a/Engine/Plugins/Experimental/ColorCorrectRegions/Source/ColorCorrectRegions/Private/ColorCorrectWindow.cpp
+++ b/Engine/Plugins/Experimental/ColorCorrectRegions/Source/ColorCorrectRegions/Private/ColorCorrectWindow.cpp
@@ -12,10 +12,6 @@
 #include "Misc/EnumRange.h"
 #include "UObject/ConstructorHelpers.h"
 
-#if WITH_EDITOR
-#include "IDisplayClusterLightCardExtenderModule.h"
-#endif
-
 ENUM_RANGE_BY_COUNT(EColorCorrectWindowType, EColorCorrectWindowType::MAX)
 
 AColorCorrectionWindow::AColorCorrectionWindow(const FObjectInitializer& ObjectInitializer)
@@ -55,25 +51,10 @@ AColorCorrectionWindow::AColorCorrectionWindow(const FObjectInitializer& ObjectI
 #if WITH_METADATA
 	CreateIcon();
 #endif
-
-#if WITH_EDITOR
-	if (!IsTemplate())
-	{
-		IDisplayClusterLightCardExtenderModule& LightCardExtenderModule = IDisplayClusterLightCardExtenderModule::Get();
-		LightCardExtenderModule.GetOnSequencerTimeChanged().AddUObject(this, &AColorCorrectionWindow::OnSequencerTimeChanged);
-	}
-#endif
 }
 
 AColorCorrectionWindow::~AColorCorrectionWindow()
 {
-#if WITH_EDITOR
-	if (!IsTemplate())
-	{
-		IDisplayClusterLightCardExtenderModule& LightCardExtenderModule = IDisplayClusterLightCardExtenderModule::Get();
-		LightCardExtenderModule.GetOnSequencerTimeChanged().RemoveAll(this);
-	}
-#endif
 }
 
 #if WITH_EDITOR
@@ -85,40 +66,13 @@ void AColorCorrectionWindow::PostEditChangeProperty(struct FPropertyChangedEvent
 	{
 		SetMeshVisibilityForWindowType();
 	}
-	else
-	{
-		const FStructProperty* StructProperty = CastField<FStructProperty>(PropertyChangedEvent.MemberProperty);
-		const bool bIsOrientation = StructProperty ? StructProperty->Struct == FDisplayClusterPositionalParams::StaticStruct() : false;
-	
-		if (bIsOrientation)
-		{
-			UpdateStageActorTransform();
-			// Updates MU in real-time. Skip our method as the positional coordinates are already correct.
-			AActor::PostEditMove(PropertyChangedEvent.ChangeType != EPropertyChangeType::Interactive);
-		}
-		else if (
-			PropertyName == USceneComponent::GetRelativeLocationPropertyName() ||
-			PropertyName == USceneComponent::GetRelativeRotationPropertyName() ||
-			PropertyName == USceneComponent::GetRelativeScale3DPropertyName())
-		{
-			bNotifyOnParamSetter = false;
-			UpdatePositionalParamsFromTransform();
-			bNotifyOnParamSetter = true;
-		}
-	}
 
-	// Call after stage actor transform is updated, so any observers will have both the correct actor transform and
-	// positional properties.
 	Super::PostEditChangeProperty(PropertyChangedEvent);
 }
 
 void AColorCorrectionWindow::PostEditMove(bool bFinished)
 {
 	Super::PostEditMove(bFinished);
-
-	bNotifyOnParamSetter = false;
-	UpdatePositionalParamsFromTransform();
-	bNotifyOnParamSetter = true;
 }
 
 FName AColorCorrectionWindow::GetCustomIconName() const
@@ -126,11 +80,6 @@ FName AColorCorrectionWindow::GetCustomIconName() const
 	return TEXT("CCW.OutlinerThumbnail");
 }
 
-void AColorCorrectionWindow::OnSequencerTimeChanged(TWeakPtr<ISequencer> InSequencer)
-{
-	UpdatePositionalParamsFromTransform();
-}
-
 #endif //WITH_EDITOR
 
 
@@ -195,151 +144,8 @@ void AColorCorrectionWindow::SetMeshVisibilityForWindowType()
 	}
 }
 
-#define NOTIFY_PARAM_SETTER()\
-	if (bNotifyOnParamSetter)\
-	{\
-		UpdateStageActorTransform();\
-	}\
-
-void AColorCorrectionWindow::SetLongitude(double InValue)
-{
-	PositionalParams.Longitude = InValue;
-	NOTIFY_PARAM_SETTER()
-}
-
-double AColorCorrectionWindow::GetLongitude() const
-{
-	return PositionalParams.Longitude;
-}
-
-void AColorCorrectionWindow::SetLatitude(double InValue)
-{
-	PositionalParams.Latitude = InValue;
-	NOTIFY_PARAM_SETTER()
-}
-
-double AColorCorrectionWindow::GetLatitude() const
-{
-	return PositionalParams.Latitude;
-}
-
-void AColorCorrectionWindow::SetDistanceFromCenter(double InValue)
-{
-	PositionalParams.DistanceFromCenter = InValue;
-	NOTIFY_PARAM_SETTER()
-}
-
-double AColorCorrectionWindow::GetDistanceFromCenter() const
-{
-	return PositionalParams.DistanceFromCenter;
-}
-
-void AColorCorrectionWindow::SetSpin(double InValue)
-{
-	PositionalParams.Spin = InValue;
-	NOTIFY_PARAM_SETTER()
-}
-
-double AColorCorrectionWindow::GetSpin() const
-{
-	return PositionalParams.Spin;
-}
-
-void AColorCorrectionWindow::SetPitch(double InValue)
-{
-	PositionalParams.Pitch = InValue;
-	NOTIFY_PARAM_SETTER()
-}
-
-double AColorCorrectionWindow::GetPitch() const
-{
-	return PositionalParams.Pitch;
-}
-
-void AColorCorrectionWindow::SetYaw(double InValue)
-{
-	PositionalParams.Yaw = InValue;
-	NOTIFY_PARAM_SETTER()
-}
-
-double AColorCorrectionWindow::GetYaw() const
-{
-	return PositionalParams.Yaw;
-}
-
-void AColorCorrectionWindow::SetRadialOffset(double InValue)
-{
-	PositionalParams.RadialOffset = InValue;
-	NOTIFY_PARAM_SETTER()
-}
-
-double AColorCorrectionWindow::GetRadialOffset() const
-{
-	return PositionalParams.RadialOffset;
-}
-
-void AColorCorrectionWindow::SetScale(const FVector2D& InScale)
-{
-	PositionalParams.Scale = InScale;
-	NOTIFY_PARAM_SETTER()
-}
-
-FVector2D AColorCorrectionWindow::GetScale() const
-{
-	return PositionalParams.Scale;
-}
-
-void AColorCorrectionWindow::SetOrigin(const FTransform& InOrigin)
-{
-	Origin = InOrigin;
-}
-
-FTransform AColorCorrectionWindow::GetOrigin() const
-{
-	return Origin;
-}
-
-void AColorCorrectionWindow::SetPositionalParams(const FDisplayClusterPositionalParams& InParams)
-{
-	PositionalParams = InParams;
-	NOTIFY_PARAM_SETTER()
-}
-
-FDisplayClusterPositionalParams AColorCorrectionWindow::GetPositionalParams() const
-{
-	return PositionalParams;
-}
-
 ADEPRECATED_ColorCorrectWindow::ADEPRECATED_ColorCorrectWindow(const FObjectInitializer& ObjectInitializer)
 	: Super(ObjectInitializer)
 {
 
 }
-
-void AColorCorrectionWindow::GetPositionalProperties(FPositionalPropertyArray& OutPropertyPairs) const
-{
-	void* Container = (void*)(&PositionalParams);
-
-	const TSet<FName>& PropertyNames = GetPositionalPropertyNames();
-	OutPropertyPairs.Reserve(PropertyNames.Num());
-
-	for (const FName& PropertyName : PropertyNames)
-	{
-		if (FProperty* Property = FindFProperty<FProperty>(FDisplayClusterPositionalParams::StaticStruct(), PropertyName))
-		{
-			OutPropertyPairs.Emplace(Container, Property);
-		}
-	}
-
-	if (FStructProperty* ParamsProperty = FindFProperty<FStructProperty>(GetClass(), GET_MEMBER_NAME_CHECKED(AColorCorrectionWindow, PositionalParams)))
-	{
-		OutPropertyPairs.Emplace((void*)this, ParamsProperty);
-	}
-}
-
-FName AColorCorrectionWindow::GetPositionalPropertiesMemberName() const
-{
-	return GET_MEMBER_NAME_CHECKED(AColorCorrectionWindow, PositionalParams);
-}
-
-#undef NOTIFY_PARAM_SETTER
diff --git a/Engine/Plugins/Experimental/ColorCorrectRegions/Source/ColorCorrectRegions/Public/ColorCorrectRegion.h b/Engine/Plugins/Experimental/ColorCorrectRegions/Source/ColorCorrectRegions/Public/ColorCorrectRegion.h
--- a/Engine/Plugins/Experimental/ColorCorrectRegions/Source/ColorCorrectRegions/Public/ColorCorrectRegion.h
+++ b/Engine/Plugins/Experimental/ColorCorrectRegions/Source/ColorCorrectRegions/Public/ColorCorrectRegion.h
@@ -7,6 +7,9 @@
 #include "Engine/Scene.h"
 #include "GameFramework/Actor.h"
 #include "UObject/ObjectMacros.h"
+
+#include "StageActor/IDisplayClusterStageActor.h"
+
 #include "ColorCorrectRegion.generated.h"
 
 
@@ -107,10 +110,12 @@ typedef TSharedPtr<FColorCorrectRenderProxy, ESPMode::ThreadSafe> FColorCorrectR
  * More information in ColorCorrectRegionsSubsytem.h
  */
 UCLASS(Blueprintable, NotPlaceable, Abstract)
-class COLORCORRECTREGIONS_API AColorCorrectRegion : public AActor
+class COLORCORRECTREGIONS_API AColorCorrectRegion : public AActor, public IDisplayClusterStageActor
 {
 	GENERATED_UCLASS_BODY()
 public:
+	virtual ~AColorCorrectRegion() override;
+	
 	/** Region type. */
 	UPROPERTY(BlueprintReadWrite, EditAnywhere, Category="Color Correction")
 	EColorCorrectRegionsType Type;
@@ -193,6 +198,7 @@ public:
 #if WITH_EDITOR
 	/** Called when any of the properties are changed. */
 	virtual void PostEditChangeProperty(FPropertyChangedEvent& PropertyChangedEvent) override;
+	virtual void PostEditMove(bool bFinished) override;
 
 #endif
 
@@ -239,6 +245,80 @@ private:
 	*/
 	void TransferState();
 
+#if WITH_EDITOR
+	/** Called when a sequencer has its time changed. */
+	void OnSequencerTimeChanged(TWeakPtr<class ISequencer> InSequencer);
+#endif
+	
+public:
+
+	// ~Begin IDisplayClusterStageActor interface
+	UFUNCTION(BlueprintCallable, Category = Orientation)
+	virtual void SetLongitude(double InValue) override;
+	UFUNCTION(BlueprintCallable, Category = Orientation)
+	virtual double GetLongitude() const override;
+
+	UFUNCTION(BlueprintCallable, Category = Orientation)
+	virtual void SetLatitude(double InValue) override;
+	UFUNCTION(BlueprintCallable, Category = Orientation)
+	virtual double GetLatitude() const override;
+
+	UFUNCTION(BlueprintCallable, Category = Orientation)
+	virtual void SetDistanceFromCenter(double InValue) override;
+	UFUNCTION(BlueprintCallable, Category = Orientation)
+	virtual double GetDistanceFromCenter() const override;
+
+	UFUNCTION(BlueprintCallable, Category = Orientation)
+	virtual void SetSpin(double InValue) override;
+	UFUNCTION(BlueprintCallable, Category = Orientation)
+	virtual double GetSpin() const override;
+
+	UFUNCTION(BlueprintCallable, Category = Orientation)
+	virtual void SetPitch(double InValue) override;
+	UFUNCTION(BlueprintCallable, Category = Orientation)
+	virtual double GetPitch() const override;
+
+	UFUNCTION(BlueprintCallable, Category = Orientation)
+	virtual void SetYaw(double InValue) override;
+	UFUNCTION(BlueprintCallable, Category = Orientation)
+	virtual double GetYaw() const override;
+
+	UFUNCTION(BlueprintCallable, Category = Orientation)
+	virtual void SetRadialOffset(double InValue) override;
+	UFUNCTION(BlueprintCallable, Category = Orientation)
+	virtual double GetRadialOffset() const override;
+
+	UFUNCTION(BlueprintCallable, Category = Orientation)
+	virtual void SetScale(const FVector2D& InScale) override;
+	UFUNCTION(BlueprintCallable, Category = Orientation)
+	virtual FVector2D GetScale() const override;
+
+	UFUNCTION(BlueprintSetter)
+	virtual void SetOrigin(const FTransform& InOrigin) override;
+	UFUNCTION(BlueprintGetter)
+	virtual FTransform GetOrigin() const override;
+
+	UFUNCTION(BlueprintCallable, Category = Orientation)
+	virtual void SetPositionalParams(const FDisplayClusterPositionalParams& InParams) override;
+	UFUNCTION(BlueprintCallable, Category = Orientation)
+	virtual FDisplayClusterPositionalParams GetPositionalParams() const override;
+
+	virtual void GetPositionalProperties(FPositionalPropertyArray& OutPropertyPairs) const override;
+	virtual FName GetPositionalPropertiesMemberName() const override;
+	// ~End IDisplayClusterStageActor interface
+	
+protected:
+	/** Spherical coordinates in relation to the origin, primarily used with the ICVFX panel. */
+	UPROPERTY(EditAnywhere, Category = Orientation, meta = (ShowOnlyInnerProperties))
+	FDisplayClusterPositionalParams PositionalParams;
+
+	/** The origin when used in the ICVFX panel. */
+	UPROPERTY(VisibleAnywhere, BlueprintSetter=SetOrigin, BlueprintGetter=GetOrigin, Category = Orientation)
+	FTransform Origin;
+
+	/** Update the transform when a positional setter is called. */
+	bool bNotifyOnParamSetter = true;
+	
 private:
 	TWeakObjectPtr<UColorCorrectRegionsSubsystem> ColorCorrectRegionsSubsystem;
 
diff --git a/Engine/Plugins/Experimental/ColorCorrectRegions/Source/ColorCorrectRegions/Public/ColorCorrectWindow.h b/Engine/Plugins/Experimental/ColorCorrectRegions/Source/ColorCorrectRegions/Public/ColorCorrectWindow.h
--- a/Engine/Plugins/Experimental/ColorCorrectRegions/Source/ColorCorrectRegions/Public/ColorCorrectWindow.h
+++ b/Engine/Plugins/Experimental/ColorCorrectRegions/Source/ColorCorrectRegions/Public/ColorCorrectWindow.h
@@ -9,8 +9,6 @@
 #include "GameFramework/Actor.h"
 #include "UObject/ObjectMacros.h"
 
-#include "StageActor/IDisplayClusterStageActor.h"
-
 #include "ColorCorrectWindow.generated.h"
 
 
@@ -22,7 +20,7 @@ class UBillboardComponent;
  * Color Correction Windows do not have Priority property and instead are sorted based on the distance from the camera.
  */
 UCLASS(Blueprintable, notplaceable)
-class COLORCORRECTREGIONS_API AColorCorrectionWindow : public AColorCorrectRegion, public IDisplayClusterStageActor
+class COLORCORRECTREGIONS_API AColorCorrectionWindow : public AColorCorrectRegion
 {
 	GENERATED_UCLASS_BODY()
 public:
@@ -40,9 +38,6 @@ public:
 	virtual void PostEditChangeProperty(FPropertyChangedEvent& PropertyChangedEvent) override;
 	virtual void PostEditMove(bool bFinished) override;
 	virtual FName GetCustomIconName() const override;
-
-private:
-	void OnSequencerTimeChanged(TWeakPtr<class ISequencer> InSequencer);
 #endif
 
 private:
@@ -52,73 +47,7 @@ private:
 
 	/** Swaps meshes for different CCW. */
 	void SetMeshVisibilityForWindowType();
-
-public:
-
-	// ~Begin IDisplayClusterStageActor interface
-	UFUNCTION(BlueprintCallable, Category = Orientation)
-	virtual void SetLongitude(double InValue) override;
-	UFUNCTION(BlueprintCallable, Category = Orientation)
-	virtual double GetLongitude() const override;
-
-	UFUNCTION(BlueprintCallable, Category = Orientation)
-	virtual void SetLatitude(double InValue) override;
-	UFUNCTION(BlueprintCallable, Category = Orientation)
-	virtual double GetLatitude() const override;
-
-	UFUNCTION(BlueprintCallable, Category = Orientation)
-	virtual void SetDistanceFromCenter(double InValue) override;
-	UFUNCTION(BlueprintCallable, Category = Orientation)
-	virtual double GetDistanceFromCenter() const override;
-
-	UFUNCTION(BlueprintCallable, Category = Orientation)
-	virtual void SetSpin(double InValue) override;
-	UFUNCTION(BlueprintCallable, Category = Orientation)
-	virtual double GetSpin() const override;
-
-	UFUNCTION(BlueprintCallable, Category = Orientation)
-	virtual void SetPitch(double InValue) override;
-	UFUNCTION(BlueprintCallable, Category = Orientation)
-	virtual double GetPitch() const override;
-
-	UFUNCTION(BlueprintCallable, Category = Orientation)
-	virtual void SetYaw(double InValue) override;
-	UFUNCTION(BlueprintCallable, Category = Orientation)
-	virtual double GetYaw() const override;
-
-	UFUNCTION(BlueprintCallable, Category = Orientation)
-	virtual void SetRadialOffset(double InValue) override;
-	UFUNCTION(BlueprintCallable, Category = Orientation)
-	virtual double GetRadialOffset() const override;
-
-	UFUNCTION(BlueprintCallable, Category = Orientation)
-	virtual void SetScale(const FVector2D& InScale) override;
-	UFUNCTION(BlueprintCallable, Category = Orientation)
-	virtual FVector2D GetScale() const override;
-
-	UFUNCTION(BlueprintSetter)
-	virtual void SetOrigin(const FTransform& InOrigin) override;
-	UFUNCTION(BlueprintGetter)
-	virtual FTransform GetOrigin() const override;
-
-	UFUNCTION(BlueprintCallable, Category = Orientation)
-	virtual void SetPositionalParams(const FDisplayClusterPositionalParams& InParams) override;
-	UFUNCTION(BlueprintCallable, Category = Orientation)
-	virtual FDisplayClusterPositionalParams GetPositionalParams() const override;
-
-	virtual void GetPositionalProperties(FPositionalPropertyArray& OutPropertyPairs) const override;
-	virtual FName GetPositionalPropertiesMemberName() const override;
-	// ~End IDisplayClusterStageActor interface
 	
-protected:
-	UPROPERTY(EditAnywhere, Category = Orientation, meta = (ShowOnlyInnerProperties))
-	FDisplayClusterPositionalParams PositionalParams;
-	
-	UPROPERTY(VisibleAnywhere, BlueprintSetter=SetOrigin, BlueprintGetter=GetOrigin, Category = Orientation)
-	FTransform Origin;
-
-	/** Update the transform when a positional setter is called */
-	bool bNotifyOnParamSetter = true;
 };
 
 UCLASS(Deprecated, Blueprintable, notplaceable, meta = (DeprecationMessage = "This is a deprecated version of Color Correct Window. Please re-create Color Correct Window to remove this warning."))
diff --git a/Engine/Plugins/Runtime/nDisplay/Source/DisplayClusterLightCardEditor/Private/DisplayClusterLightCardEditorUtils.cpp b/Engine/Plugins/Runtime/nDisplay/Source/DisplayClusterLightCardEditor/Private/DisplayClusterLightCardEditorUtils.cpp
--- a/Engine/Plugins/Runtime/nDisplay/Source/DisplayClusterLightCardEditor/Private/DisplayClusterLightCardEditorUtils.cpp
+++ b/Engine/Plugins/Runtime/nDisplay/Source/DisplayClusterLightCardEditor/Private/DisplayClusterLightCardEditorUtils.cpp
@@ -28,7 +28,7 @@ TSet<UClass*> UE::DisplayClusterLightCardEditorUtils::GetAllStageActorClasses()
 	auto ShouldAddNativeClass = [] (const UClass* Class) -> bool
 	{
 		return Class->ImplementsInterface(UDisplayClusterStageActor::StaticClass()) && Class->IsNative() &&
-			!Class->HasAnyClassFlags(CLASS_Deprecated | CLASS_NewerVersionExists) &&
+			!Class->HasAnyClassFlags(CLASS_Deprecated | CLASS_NewerVersionExists | CLASS_Abstract) &&
 			!Class->GetName().StartsWith(TEXT("SKEL_")) && !Class->GetName().StartsWith(TEXT("REINST_"));
 	};
 
diff --git a/Engine/Plugins/Runtime/nDisplay/Source/DisplayClusterLightCardEditor/Private/Viewport/DisplayClusterLightCardEditorViewportClient.cpp b/Engine/Plugins/Runtime/nDisplay/Source/DisplayClusterLightCardEditor/Private/Viewport/DisplayClusterLightCardEditorViewportClient.cpp
--- a/Engine/Plugins/Runtime/nDisplay/Source/DisplayClusterLightCardEditor/Private/Viewport/DisplayClusterLightCardEditorViewportClient.cpp
+++ b/Engine/Plugins/Runtime/nDisplay/Source/DisplayClusterLightCardEditor/Private/Viewport/DisplayClusterLightCardEditorViewportClient.cpp
@@ -1762,22 +1762,24 @@ void FDisplayClusterLightCardEditorViewportClient::UpdatePreviewActor(ADisplayCl
 			for (const TObjectPtr<AActor>& ActorProxy : ActorProxiesCreated)
 			{
 				// Hack - CL 23230783 sets CCW meshes to hidden which causes problems with selection, so always add CCWs to the renderer.
-				bool bIsCCW = false;
+				// Updating to include all stage actors now that CCRs are also supported.
+				bool bIsStageActor = false;
 				for (const UClass* Class = ActorProxy->GetClass(); Class && (UObject::StaticClass() != Class); Class = Class->GetSuperClass())
 				{
-					if (Class->GetName() == TEXT("ColorCorrectionWindow"))
+					if (Class->ImplementsInterface(UDisplayClusterStageActor::StaticClass()))
 					{
-						bIsCCW = true;
+						bIsStageActor = true;
 						break;
 					}
 				}
 				
-				IDisplayClusterScenePreview::Get().AddActorToRenderer(PreviewRendererId, ActorProxy, [this, ActorProxy, bIsCCW](const UPrimitiveComponent* PrimitiveComponent)
+				IDisplayClusterScenePreview::Get().AddActorToRenderer(PreviewRendererId,
+					ActorProxy, [this, ActorProxy, bIsStageActor](const UPrimitiveComponent* PrimitiveComponent)
 				{
 					// Always add the light card mesh component to the renderer's scene even if it is marked hidden in game, since UV light cards will purposefully
 					// hide the light card mesh since it isn't supposed to exist in 3D space. The light card mesh will be appropriately filtered when the scene is
 					// rendered based on the projection mode
-					if (PrimitiveComponent->GetFName() == TEXT("LightCard") || bIsCCW)
+					if (PrimitiveComponent->GetFName() == TEXT("LightCard") || bIsStageActor)
 					{
 						return true;
 					}
