Color Correction Regions: Are no longer hidden in Editor when Hidden In Game is enabled

diff --git a/Engine/Plugins/Experimental/ColorCorrectRegions/Source/ColorCorrectRegions/Private/ColorCorrectRegion.cpp b/Engine/Plugins/Experimental/ColorCorrectRegions/Source/ColorCorrectRegions/Private/ColorCorrectRegion.cpp
--- a/Engine/Plugins/Experimental/ColorCorrectRegions/Source/ColorCorrectRegions/Private/ColorCorrectRegion.cpp
+++ b/Engine/Plugins/Experimental/ColorCorrectRegions/Source/ColorCorrectRegions/Private/ColorCorrectRegion.cpp
@@ -158,7 +158,7 @@ void AColorCorrectRegion::TransferState()
 #if WITH_EDITOR
 		&& !IsHiddenEd()
 #endif 
-		&& !IsHidden();
+		&& !(GetWorld()->HasBegunPlay() && IsHidden());
 
 	TempCCRStateRenderThread->World = GetWorld();
 	TempCCRStateRenderThread->Priority = Priority;
