Color Correct: Disable notify param on setter OnSequencerTimeChanged so that there's no infinite loop on autokey

diff --git a/Engine/Plugins/Experimental/ColorCorrectRegions/Source/ColorCorrectRegions/Private/ColorCorrectRegion.cpp b/Engine/Plugins/Experimental/ColorCorrectRegions/Source/ColorCorrectRegions/Private/ColorCorrectRegion.cpp
--- a/Engine/Plugins/Experimental/ColorCorrectRegions/Source/ColorCorrectRegions/Private/ColorCorrectRegion.cpp
+++ b/Engine/Plugins/Experimental/ColorCorrectRegions/Source/ColorCorrectRegions/Private/ColorCorrectRegion.cpp
@@ -227,7 +227,9 @@ void AColorCorrectRegion::TransferState()
 #if WITH_EDITOR
 void AColorCorrectRegion::OnSequencerTimeChanged(TWeakPtr<ISequencer> InSequencer)
 {
+	bNotifyOnParamSetter = false;
 	UpdatePositionalParamsFromTransform();
+	bNotifyOnParamSetter = true;
 }
 #endif
 
