Color Correction Regions: Fixed bug: When CCR is part of sublevel it's state is not transfered and it could remain visible/hidden despite sublevel that contains CCR having an opposite state.

diff --git a/Engine/Plugins/Experimental/ColorCorrectRegions/Source/ColorCorrectRegions/Private/ColorCorrectRegionSubsystem.cpp b/Engine/Plugins/Experimental/ColorCorrectRegions/Source/ColorCorrectRegions/Private/ColorCorrectRegionSubsystem.cpp
--- a/Engine/Plugins/Experimental/ColorCorrectRegions/Source/ColorCorrectRegions/Private/ColorCorrectRegionSubsystem.cpp
+++ b/Engine/Plugins/Experimental/ColorCorrectRegions/Source/ColorCorrectRegions/Private/ColorCorrectRegionSubsystem.cpp
@@ -326,6 +326,10 @@ void UColorCorrectRegionsSubsystem::RefreshRegions()
 		AColorCorrectRegion* AsRegion = *It;
 		if (IsRegionValid(AsRegion, GetWorld()))
 		{
+			// OnLevelChanged is called when Sublevel is hidden. State needs to be transferred
+			// once per actor on refresh since actors don't tick if sublevel is hidden.
+			AsRegion->TransferState();
+
 			if (!Cast<AColorCorrectionWindow>(AsRegion))
 			{
 				RegionsPriorityBased.Add(AsRegion);
