Color Correction Actors: A refactor and a fix for the loss of component references on loading of level snapshot.

diff --git a/Engine/Plugins/Experimental/ColorCorrectRegions/Source/ColorCorrectRegions/Private/ColorCorrectRegion.cpp b/Engine/Plugins/Experimental/ColorCorrectRegions/Source/ColorCorrectRegions/Private/ColorCorrectRegion.cpp
--- a/Engine/Plugins/Experimental/ColorCorrectRegions/Source/ColorCorrectRegions/Private/ColorCorrectRegion.cpp
+++ b/Engine/Plugins/Experimental/ColorCorrectRegions/Source/ColorCorrectRegions/Private/ColorCorrectRegion.cpp
@@ -362,13 +362,6 @@ void AColorCorrectRegion::PostEditChangeProperty(struct FPropertyChangedEvent& P
 			ColorCorrectRegionsSubsystem->SortRegionsByPriority();
 		}
 	}
-	if (PropertyName == GET_MEMBER_NAME_CHECKED(AColorCorrectRegion, Type) || PropertyChangedEvent.Property == nullptr)
-	{
-		if (ColorCorrectRegionsSubsystem.IsValid())
-		{
-			ColorCorrectRegionsSubsystem->OnLevelsChanged();
-		}
-	}
 
 	// Stage actor properties
 	{
@@ -405,6 +398,12 @@ void AColorCorrectRegion::PostEditMove(bool bFinished)
 	UpdatePositionalParamsFromTransform();
 	bNotifyOnParamSetter = true;
 }
+
+void AColorCorrectRegion::PostTransacted(const FTransactionObjectEvent& TransactionEvent)
+{
+	Super::PostTransacted(TransactionEvent);
+	FixMeshComponentReferences();
+}
 #endif //WITH_EDITOR
 
 #define NOTIFY_PARAM_SETTER()\
@@ -572,26 +571,12 @@ AColorCorrectionRegion::AColorCorrectionRegion(const FObjectInitializer& ObjectI
 		MeshComponent->CastShadow = false;
 		MeshComponent->SetHiddenInGame(true);
 	}
-
-	SetMeshVisibilityForRegionType();
-
+	ChangeShapeVisibilityForActorType();
 }
 
-void AColorCorrectionRegion::SetMeshVisibilityForRegionType()
+void AColorCorrectionRegion::ChangeShapeVisibilityForActorType()
 {
-	for (EColorCorrectRegionsType CCRType : TEnumRange<EColorCorrectRegionsType>())
-	{
-		uint8 TypeIndex = static_cast<uint8>(CCRType);
-
-		if (CCRType == Type)
-		{
-			MeshComponents[TypeIndex]->SetVisibility(true, true);
-		}
-		else
-		{
-			MeshComponents[TypeIndex]->SetVisibility(false, true);
-		}
-	}
+	ChangeShapeVisibilityForActorTypeInternal<EColorCorrectRegionsType>(Type);
 }
 
 #if WITH_EDITOR
@@ -600,9 +585,9 @@ void AColorCorrectionRegion::PostEditChangeProperty(struct FPropertyChangedEvent
 	Super::PostEditChangeProperty(PropertyChangedEvent);
 	const FName PropertyName = PropertyChangedEvent.GetPropertyName();
 
-	if (PropertyName == GET_MEMBER_NAME_CHECKED(AColorCorrectionRegion, Type))
+	if (PropertyName == GET_MEMBER_NAME_CHECKED(AColorCorrectRegion, Type) || PropertyChangedEvent.Property == nullptr)
 	{
-		SetMeshVisibilityForRegionType();
+		ChangeShapeVisibilityForActorType();
 	}
 }
 
@@ -610,6 +595,11 @@ FName AColorCorrectionRegion::GetCustomIconName() const
 {
 	return TEXT("CCR.OutlinerThumbnail");
 }
+
+void AColorCorrectionRegion::FixMeshComponentReferences()
+{
+	FixMeshComponentReferencesInternal<EColorCorrectRegionsType>(Type);
+}
 #endif
 
 #undef NOTIFY_PARAM_SETTER
diff --git a/Engine/Plugins/Experimental/ColorCorrectRegions/Source/ColorCorrectRegions/Private/ColorCorrectRegionSubsystem.cpp b/Engine/Plugins/Experimental/ColorCorrectRegions/Source/ColorCorrectRegions/Private/ColorCorrectRegionSubsystem.cpp
--- a/Engine/Plugins/Experimental/ColorCorrectRegions/Source/ColorCorrectRegions/Private/ColorCorrectRegionSubsystem.cpp
+++ b/Engine/Plugins/Experimental/ColorCorrectRegions/Source/ColorCorrectRegions/Private/ColorCorrectRegionSubsystem.cpp
@@ -145,7 +145,6 @@ void UColorCorrectRegionsSubsystem::OnActorSpawned(AActor* InActor)
 	if (IsRegionValid(AsRegion, GetWorld()))
 	{
 		FScopeLock RegionScopeLock(&RegionAccessCriticalSection);
-		EColorCorrectRegionsType CCRType = AsRegion->Type;
 		// We wouldn't have to do a check here except in case of nDisplay we need to populate this list during OnLevelsChanged 
 		// because nDisplay can release Actors while those are marked as BeginningPlay. Therefore we want to avoid 
 		// adding regions twice.
diff --git a/Engine/Plugins/Experimental/ColorCorrectRegions/Source/ColorCorrectRegions/Private/ColorCorrectWindow.cpp b/Engine/Plugins/Experimental/ColorCorrectRegions/Source/ColorCorrectRegions/Private/ColorCorrectWindow.cpp
--- a/Engine/Plugins/Experimental/ColorCorrectRegions/Source/ColorCorrectRegions/Private/ColorCorrectWindow.cpp
+++ b/Engine/Plugins/Experimental/ColorCorrectRegions/Source/ColorCorrectRegions/Private/ColorCorrectWindow.cpp
@@ -46,7 +46,8 @@ AColorCorrectionWindow::AColorCorrectionWindow(const FObjectInitializer& ObjectI
 		MeshComponent->CastShadow = false;
 		MeshComponent->SetHiddenInGame(true);
 	}
-	SetMeshVisibilityForWindowType();
+
+	ChangeShapeVisibilityForActorType();
 
 #if WITH_METADATA
 	CreateIcon();
@@ -64,7 +65,7 @@ void AColorCorrectionWindow::PostEditChangeProperty(struct FPropertyChangedEvent
 
 	if (PropertyName == GET_MEMBER_NAME_CHECKED(AColorCorrectionWindow, WindowType))
 	{
-		SetMeshVisibilityForWindowType();
+		ChangeShapeVisibilityForActorType();
 	}
 
 	Super::PostEditChangeProperty(PropertyChangedEvent);
@@ -127,22 +128,17 @@ void AColorCorrectionWindow::CreateIcon()
 }
 #endif 
 
-void AColorCorrectionWindow::SetMeshVisibilityForWindowType()
+void AColorCorrectionWindow::ChangeShapeVisibilityForActorType()
 {
-	for (EColorCorrectWindowType CCWType : TEnumRange<EColorCorrectWindowType>())
-	{
-		uint8 TypeIndex = static_cast<uint8>(CCWType);
+	ChangeShapeVisibilityForActorTypeInternal<EColorCorrectWindowType>(WindowType);
+}
 
-		if (CCWType == WindowType)
-		{
-			MeshComponents[TypeIndex]->SetVisibility(true, true);
-		}
-		else
-		{
-			MeshComponents[TypeIndex]->SetVisibility(false, true);
-		}
-	}
+#if WITH_EDITOR
+void AColorCorrectionWindow::FixMeshComponentReferences()
+{
+	FixMeshComponentReferencesInternal<EColorCorrectWindowType>(WindowType);
 }
+#endif
 
 ADEPRECATED_ColorCorrectWindow::ADEPRECATED_ColorCorrectWindow(const FObjectInitializer& ObjectInitializer)
 	: Super(ObjectInitializer)
diff --git a/Engine/Plugins/Experimental/ColorCorrectRegions/Source/ColorCorrectRegions/Public/ColorCorrectRegion.h b/Engine/Plugins/Experimental/ColorCorrectRegions/Source/ColorCorrectRegions/Public/ColorCorrectRegion.h
--- a/Engine/Plugins/Experimental/ColorCorrectRegions/Source/ColorCorrectRegions/Public/ColorCorrectRegion.h
+++ b/Engine/Plugins/Experimental/ColorCorrectRegions/Source/ColorCorrectRegions/Public/ColorCorrectRegion.h
@@ -2,11 +2,12 @@
 
 #pragma once
 
-#include "Components/MeshComponent.h"
 #include "CoreMinimal.h"
 #include "Engine/Scene.h"
 #include "GameFramework/Actor.h"
 #include "UObject/ObjectMacros.h"
+#include "UObject/Object.h"
+#include "Components/StaticMeshComponent.h"
 
 #include "StageActor/IDisplayClusterStageActor.h"
 
@@ -113,6 +114,10 @@ UCLASS(Blueprintable, NotPlaceable, Abstract)
 class COLORCORRECTREGIONS_API AColorCorrectRegion : public AActor, public IDisplayClusterStageActor
 {
 	GENERATED_UCLASS_BODY()
+public:
+	UPROPERTY()
+	TArray<TObjectPtr<UStaticMeshComponent>> MeshComponents;
+
 public:
 	virtual ~AColorCorrectRegion() override;
 	
@@ -199,9 +204,10 @@ public:
 	/** Called when any of the properties are changed. */
 	virtual void PostEditChangeProperty(FPropertyChangedEvent& PropertyChangedEvent) override;
 	virtual void PostEditMove(bool bFinished) override;
-
+	virtual void PostTransacted(const FTransactionObjectEvent& TransactionEvent) override;
 #endif
 
+public:
 	/** The main purpose of this component is to determine the visibility status of this Color Correction Actor. */
 	UPROPERTY()
 	TObjectPtr<UColorCorrectionInvisibleComponent> IdentityComponent;
@@ -231,6 +237,71 @@ public:
 	*/
 	void TransferState();
 
+	/** 
+	* All CC Actors rely on a shape for selection. These shapes need to be swapped depending on the type.
+	* This function forces the refresh of the shape if the Type was changed outside of UI.
+	* For internal use only.
+	*/
+	virtual void ChangeShapeVisibilityForActorType() {};
+
+protected:
+
+	/** All CC Actors rely on a shape for selection. These shapes need to be swapped depending on the type. */
+	template <typename TCCActorType>
+	void ChangeShapeVisibilityForActorTypeInternal(TCCActorType InDesiredType)
+	{
+		for (TCCActorType CCActorType : TEnumRange<TCCActorType>())
+		{
+			uint8 TypeIndex = static_cast<uint8>(CCActorType);
+
+			ensure(IsValid(MeshComponents[TypeIndex]));
+
+			if (CCActorType == InDesiredType)
+			{
+				MeshComponents[TypeIndex]->SetVisibility(true, true);
+			}
+			else
+			{
+				MeshComponents[TypeIndex]->SetVisibility(false, true);
+			}
+		}
+	};
+
+#if WITH_EDITOR
+protected:
+	/** Used to validate child components. */
+	virtual void FixMeshComponentReferences() {};
+
+	template <typename TCCActorType>
+	void FixMeshComponentReferencesInternal(TCCActorType InDesiredType)
+	{
+		if (!RootComponent)
+		{
+			return;
+		}
+
+		TArray<USceneComponent*> ChildComponents;
+		RootComponent->GetChildrenComponents(false, ChildComponents);
+		MeshComponents.Empty();
+		for (TCCActorType CCActorType : TEnumRange<TCCActorType>())
+		{
+			const FString ShapeNameString = UEnum::GetValueAsString(CCActorType);
+			for (USceneComponent* ChildComponent : ChildComponents)
+			{
+				TObjectPtr<UStaticMeshComponent> ChildMeshComponent = Cast<UStaticMeshComponent>(ChildComponent);
+				if (ChildMeshComponent && ChildComponent->GetName() == ShapeNameString)
+				{
+					MeshComponents.Add(ChildMeshComponent);
+					break;
+				}
+
+			}
+
+		}
+		ChangeShapeVisibilityForActorTypeInternal<TCCActorType>(InDesiredType);
+	};
+#endif
+
 private:
 
 #if WITH_METADATA
@@ -343,19 +414,19 @@ UCLASS(Blueprintable, NotPlaceable)
 class COLORCORRECTREGIONS_API AColorCorrectionRegion : public AColorCorrectRegion
 {
 	GENERATED_UCLASS_BODY()
-public:
-	UPROPERTY()
-	TArray<TObjectPtr<UStaticMeshComponent>> MeshComponents;
-
-public:
-	/** Swaps meshes for different CCR. */
-	void SetMeshVisibilityForRegionType();
 
 #if WITH_EDITOR
 	/** Called when any of the properties are changed. */
 	virtual void PostEditChangeProperty(FPropertyChangedEvent& PropertyChangedEvent) override;
 	virtual FName GetCustomIconName() const override;
+
+protected:
+	virtual void FixMeshComponentReferences() override;
 #endif
+
+protected:
+	virtual void ChangeShapeVisibilityForActorType() override;
+
 };
 
 
diff --git a/Engine/Plugins/Experimental/ColorCorrectRegions/Source/ColorCorrectRegions/Public/ColorCorrectWindow.h b/Engine/Plugins/Experimental/ColorCorrectRegions/Source/ColorCorrectRegions/Public/ColorCorrectWindow.h
--- a/Engine/Plugins/Experimental/ColorCorrectRegions/Source/ColorCorrectRegions/Public/ColorCorrectWindow.h
+++ b/Engine/Plugins/Experimental/ColorCorrectRegions/Source/ColorCorrectRegions/Public/ColorCorrectWindow.h
@@ -26,9 +26,6 @@ class COLORCORRECTREGIONS_API AColorCorrectionWindow : public AColorCorrectRegio
 public:
 	virtual ~AColorCorrectionWindow() override;
 
-	UPROPERTY()
-	TArray<TObjectPtr<UStaticMeshComponent>> MeshComponents;
-
 	/** Region type. */
 	UPROPERTY(BlueprintReadWrite, EditAnywhere, Category="Color Correction", Meta = (DisplayName = "Type"))
 	EColorCorrectWindowType WindowType;
@@ -38,16 +35,20 @@ public:
 	virtual void PostEditChangeProperty(FPropertyChangedEvent& PropertyChangedEvent) override;
 	virtual void PostEditMove(bool bFinished) override;
 	virtual FName GetCustomIconName() const override;
+
+protected:
+	virtual void FixMeshComponentReferences() override;
 #endif
 
+protected:
+	virtual void ChangeShapeVisibilityForActorType() override;
+
 private:
 #if WITH_METADATA
 	void CreateIcon();
 #endif // WITH_METADATA
 
-	/** Swaps meshes for different CCW. */
-	void SetMeshVisibilityForWindowType();
-	
+
 };
 
 UCLASS(Deprecated, Blueprintable, notplaceable, meta = (DeprecationMessage = "This is a deprecated version of Color Correct Window. Please re-create Color Correct Window to remove this warning."))
diff --git a/Engine/Plugins/Experimental/ColorCorrectRegions/Source/ColorCorrectRegionsEditor/Private/ColorCorrectionActorContextMenu.cpp b/Engine/Plugins/Experimental/ColorCorrectRegions/Source/ColorCorrectRegionsEditor/Private/ColorCorrectionActorContextMenu.cpp
--- a/Engine/Plugins/Experimental/ColorCorrectRegions/Source/ColorCorrectRegionsEditor/Private/ColorCorrectionActorContextMenu.cpp
+++ b/Engine/Plugins/Experimental/ColorCorrectRegions/Source/ColorCorrectRegionsEditor/Private/ColorCorrectionActorContextMenu.cpp
@@ -111,10 +111,11 @@ namespace
 		{
 			AColorCorrectionRegion* CCActorRawPtr = World->SpawnActor<AColorCorrectionRegion>();
 			CCActorRawPtr->Type = EColorCorrectRegionsType::Box;
-			CCActorRawPtr->SetMeshVisibilityForRegionType();
 			CCActorPtr = CCActorRawPtr;
 		}
 
+		CCActorPtr->ChangeShapeVisibilityForActorType();
+
 		CCActorPtr->SetActorTransform(Transform);
 		AddActorsToPerActorCC(CCActorPtr, SelectedActors);
 
