Color Correction Actors: Minor refactor. Made a method private.

diff --git a/Engine/Plugins/Experimental/ColorCorrectRegions/Source/ColorCorrectRegions/Public/ColorCorrectRegion.h b/Engine/Plugins/Experimental/ColorCorrectRegions/Source/ColorCorrectRegions/Public/ColorCorrectRegion.h
--- a/Engine/Plugins/Experimental/ColorCorrectRegions/Source/ColorCorrectRegions/Public/ColorCorrectRegion.h
+++ b/Engine/Plugins/Experimental/ColorCorrectRegions/Source/ColorCorrectRegions/Public/ColorCorrectRegion.h
@@ -237,15 +237,15 @@ public:
 	*/
 	void TransferState();
 
-	/** 
+protected:
+
+	/**
 	* All CC Actors rely on a shape for selection. These shapes need to be swapped depending on the type.
 	* This function forces the refresh of the shape if the Type was changed outside of UI.
 	* For internal use only.
 	*/
 	virtual void ChangeShapeVisibilityForActorType() {};
 
-protected:
-
 	/** All CC Actors rely on a shape for selection. These shapes need to be swapped depending on the type. */
 	template <typename TCCActorType>
 	void ChangeShapeVisibilityForActorTypeInternal(TCCActorType InDesiredType)
diff --git a/Engine/Plugins/Experimental/ColorCorrectRegions/Source/ColorCorrectRegionsEditor/Private/ColorCorrectionActorContextMenu.cpp b/Engine/Plugins/Experimental/ColorCorrectRegions/Source/ColorCorrectRegionsEditor/Private/ColorCorrectionActorContextMenu.cpp
--- a/Engine/Plugins/Experimental/ColorCorrectRegions/Source/ColorCorrectRegionsEditor/Private/ColorCorrectionActorContextMenu.cpp
+++ b/Engine/Plugins/Experimental/ColorCorrectRegions/Source/ColorCorrectRegionsEditor/Private/ColorCorrectionActorContextMenu.cpp
@@ -112,9 +112,10 @@ namespace
 			AColorCorrectionRegion* CCActorRawPtr = World->SpawnActor<AColorCorrectionRegion>();
 			CCActorRawPtr->Type = EColorCorrectRegionsType::Box;
 			CCActorPtr = CCActorRawPtr;
-		}
 
-		CCActorPtr->ChangeShapeVisibilityForActorType();
+			FPropertyChangedEvent TypeChangedEvent(AColorCorrectRegion::StaticClass()->FindPropertyByName(GET_MEMBER_NAME_CHECKED(AColorCorrectRegion, Type)));
+			CCActorPtr->PostEditChangeProperty(TypeChangedEvent);
+		}
 
 		CCActorPtr->SetActorTransform(Transform);
 		AddActorsToPerActorCC(CCActorPtr, SelectedActors);
