Color Correction Actors: Made the popup that wars users of the change of stencil on actors disabled based on a CVarEnablePopupStencilIdChangedOnActors

diff --git a/Engine/Plugins/Experimental/ColorCorrectRegions/Source/ColorCorrectRegions/Private/ColorCorrectRegionsStencilManager.cpp b/Engine/Plugins/Experimental/ColorCorrectRegions/Source/ColorCorrectRegions/Private/ColorCorrectRegionsStencilManager.cpp
--- a/Engine/Plugins/Experimental/ColorCorrectRegions/Source/ColorCorrectRegions/Private/ColorCorrectRegionsStencilManager.cpp
+++ b/Engine/Plugins/Experimental/ColorCorrectRegions/Source/ColorCorrectRegions/Private/ColorCorrectRegionsStencilManager.cpp
@@ -22,6 +22,12 @@ static TAutoConsoleVariable<int32> CVarStencilIdRangeMax(
 	TEXT("The upper boundary of stencil Ids that are reserved by CCR. \n\
 		By default this value is set to 255. Must be within the range of [0, 255].\n"));
 
+static TAutoConsoleVariable<bool> CVarEnablePopupStencilIdChangedOnActors(
+	TEXT("r.CCR.EnablePopupStencilIdChangedOnActors"),
+	false,
+	TEXT("Allows CCRs to notify users of change in stencil Ids on actors modified by CCR. \n\
+		By default no notification is displayed and CCR assumes that users changed stencil numbers intentionally.\n"));
+
 namespace
 {
 	/*
@@ -408,7 +414,11 @@ void FColorCorrectRegionsStencilManager::CheckAssignedActorsValidity(AColorCorre
 					if (PrimitiveComponent->CustomDepthStencilValue != StencilDataPair.Value->AssignedStencil || !PrimitiveComponent->bRenderCustomDepth)
 					{
 						auto Text = FString::Printf(TEXT("Custom Stencil Id for Actor: '%s', Component: '%s' is managed by Color Correct Region '%s' and has been modified manually. Would you like to change the id back?"), *Actor->GetName(), *PrimitiveComponent->GetName(), *Region->GetName());
-						EAppReturnType::Type Answer = FMessageDialog::Open(EAppMsgType::YesNo, FText::FromString(Text));
+						EAppReturnType::Type Answer = EAppReturnType::No;
+						if (CVarEnablePopupStencilIdChangedOnActors.GetValueOnAnyThread())
+						{
+							Answer = FMessageDialog::Open(EAppMsgType::YesNo, FText::FromString(Text));
+						}
 						if (Answer == EAppReturnType::No)
 						{
 							// Update our stored value and then update other components with the new stencil Id value.
diff --git a/Engine/Plugins/Experimental/ColorCorrectRegions/Source/ColorCorrectRegions/Public/ColorCorrectRegion.h b/Engine/Plugins/Experimental/ColorCorrectRegions/Source/ColorCorrectRegions/Public/ColorCorrectRegion.h
--- a/Engine/Plugins/Experimental/ColorCorrectRegions/Source/ColorCorrectRegions/Public/ColorCorrectRegion.h
+++ b/Engine/Plugins/Experimental/ColorCorrectRegions/Source/ColorCorrectRegions/Public/ColorCorrectRegion.h
@@ -254,7 +254,15 @@ protected:
 		{
 			uint8 TypeIndex = static_cast<uint8>(CCActorType);
 
-			ensure(IsValid(MeshComponents[TypeIndex]));
+			if (!IsValid(MeshComponents[TypeIndex]))
+			{
+#if WITH_EDITOR
+				FixMeshComponentReferencesInternal<TCCActorType>(InDesiredType);
+#else
+				ensure(IsValid(MeshComponents[TypeIndex]));
+#endif
+				return;
+			}
 
 			if (CCActorType == InDesiredType)
 			{
