Color Correct Actors: Added a sanity check for validating meshes.

diff --git a/Engine/Plugins/Experimental/ColorCorrectRegions/Source/ColorCorrectRegions/Public/ColorCorrectRegion.h b/Engine/Plugins/Experimental/ColorCorrectRegions/Source/ColorCorrectRegions/Public/ColorCorrectRegion.h
--- a/Engine/Plugins/Experimental/ColorCorrectRegions/Source/ColorCorrectRegions/Public/ColorCorrectRegion.h
+++ b/Engine/Plugins/Experimental/ColorCorrectRegions/Source/ColorCorrectRegions/Public/ColorCorrectRegion.h
@@ -250,6 +250,11 @@ protected:
 	template <typename TCCActorType>
 	void ChangeShapeVisibilityForActorTypeInternal(TCCActorType InDesiredType)
 	{
+		if (!IsValid(this) || MeshComponents.Num() != (uint8)TCCActorType::MAX)
+		{
+			return;
+		}
+
 		for (TCCActorType CCActorType : TEnumRange<TCCActorType>())
 		{
 			uint8 TypeIndex = static_cast<uint8>(CCActorType);
