Color Correction Actors: Removed private dependencies.

diff --git a/Engine/Plugins/Experimental/ColorCorrectRegions/Source/ColorCorrectRegions/ColorCorrectRegions.Build.cs b/Engine/Plugins/Experimental/ColorCorrectRegions/Source/ColorCorrectRegions/ColorCorrectRegions.Build.cs
--- a/Engine/Plugins/Experimental/ColorCorrectRegions/Source/ColorCorrectRegions/ColorCorrectRegions.Build.cs
+++ b/Engine/Plugins/Experimental/ColorCorrectRegions/Source/ColorCorrectRegions/ColorCorrectRegions.Build.cs
@@ -7,13 +7,6 @@ public class ColorCorrectRegions : ModuleRules
 {
 	public ColorCorrectRegions(ReadOnlyTargetRules Target) : base(Target)
 	{
-		PrivateIncludePaths.AddRange(
-			new string[] {
-				//Used in ColorCorrectRegionsPostProcessMaterial.h
-				Path.Combine(GetModuleDirectory("Renderer"), "Private"),
-			}
-		);
-		
 		PublicDependencyModuleNames.AddRange(
 			new string[]
 			{
diff --git a/Engine/Plugins/Experimental/ColorCorrectRegions/Source/ColorCorrectRegions/Private/ColorCorrectRegionsSceneViewExtension.cpp b/Engine/Plugins/Experimental/ColorCorrectRegions/Source/ColorCorrectRegions/Private/ColorCorrectRegionsSceneViewExtension.cpp
--- a/Engine/Plugins/Experimental/ColorCorrectRegions/Source/ColorCorrectRegions/Private/ColorCorrectRegionsSceneViewExtension.cpp
+++ b/Engine/Plugins/Experimental/ColorCorrectRegions/Source/ColorCorrectRegions/Private/ColorCorrectRegionsSceneViewExtension.cpp
@@ -9,12 +9,12 @@
 #include "DynamicResolutionState.h"
 #include "Engine/World.h"
 #include "EngineUtils.h"
+#include "FXRenderingUtils.h"
+#include "PostProcess/PostProcessInputs.h"
 #include "RHI.h"
-#include "ScreenPass.h"
+#include "SceneRenderTargetParameters.h"
 #include "SceneView.h"
-
-#include "PostProcess/PostProcessing.h"
-#include "SceneRendering.h"
+#include "ScreenPass.h"
 
 // Set this to 1 to clip pixels outside of bounding box.
 #define CLIP_PIXELS_OUTSIDE_AABB 1
@@ -668,17 +668,18 @@ void FColorCorrectRegionsSceneViewExtension::PrePostProcessPass_RenderThread(FRD
 
 	const FSceneViewFamily& ViewFamily = *View.Family;
 
-	DynamicRenderScaling::TMap<float> UpperBounds = ViewFamily.GetScreenPercentageInterface()->GetResolutionFractionsUpperBound();
-	const auto FeatureLevel = View.GetFeatureLevel();
-	const float ScreenPercentage = UpperBounds[GDynamicPrimaryResolutionFraction] * ViewFamily.SecondaryViewFraction;
-	
 	// We need to make sure to take Windows and Scene scale into account.
+	float ScreenPercentage = ViewFamily.SecondaryViewFraction;
 
-	checkSlow(View.bIsViewInfo); // can't do dynamic_cast because FViewInfo doesn't have any virtual functions.
-	const FIntRect PrimaryViewRect = static_cast<const FViewInfo&>(View).ViewRect;
+	if (ViewFamily.GetScreenPercentageInterface())
+	{
+		DynamicRenderScaling::TMap<float> UpperBounds = ViewFamily.GetScreenPercentageInterface()->GetResolutionFractionsUpperBound();
+		ScreenPercentage *= UpperBounds[GDynamicPrimaryResolutionFraction];
+	}
 
-	FScreenPassTexture SceneColor((*Inputs.SceneTextures)->SceneColorTexture, PrimaryViewRect);
+	const FIntRect PrimaryViewRect = UE::FXRenderingUtils::GetRawViewRectUnsafe(View);
 
+	FScreenPassTexture SceneColor((*Inputs.SceneTextures)->SceneColorTexture, PrimaryViewRect);
 
 	if (!SceneColor.IsValid())
 	{
@@ -716,7 +717,6 @@ void FColorCorrectRegionsSceneViewExtension::PrePostProcessPass_RenderThread(FRD
 
 		// Because we are not using proxy material, but plain global shader, we need to setup Scene textures ourselves.
 		// We don't need to do this per region.
-		check(View.bIsViewInfo);
 		FSceneTextureShaderParameters SceneTextures = CreateSceneTextureShaderParameters(GraphBuilder, View, ESceneTextureSetupMode::All);
 
 		WorldSubsystem->SortRegionsByDistance(View.ViewLocation);
