Color Correction Regions: Fixed an issue with Per actor CC when CC actors are on sublevel that is being toggled on and off.

diff --git a/Engine/Plugins/Experimental/ColorCorrectRegions/Source/ColorCorrectRegions/Private/ColorCorrectRegion.cpp b/Engine/Plugins/Experimental/ColorCorrectRegions/Source/ColorCorrectRegions/Private/ColorCorrectRegion.cpp
--- a/Engine/Plugins/Experimental/ColorCorrectRegions/Source/ColorCorrectRegions/Private/ColorCorrectRegion.cpp
+++ b/Engine/Plugins/Experimental/ColorCorrectRegions/Source/ColorCorrectRegions/Private/ColorCorrectRegion.cpp
@@ -86,7 +86,7 @@ void AColorCorrectRegion::EndPlay(const EEndPlayReason::Type EndPlayReason)
 {
 	if (ColorCorrectRegionsSubsystem.IsValid())
 	{
-		ColorCorrectRegionsSubsystem->OnActorDeleted(this);
+		ColorCorrectRegionsSubsystem->OnActorDeleted(this, false);
 		ColorCorrectRegionsSubsystem = nullptr;
 	}
 	Super::EndPlay(EndPlayReason);
@@ -96,7 +96,7 @@ void AColorCorrectRegion::BeginDestroy()
 {
 	if (ColorCorrectRegionsSubsystem.IsValid())
 	{
-		ColorCorrectRegionsSubsystem->OnActorDeleted(this);
+		ColorCorrectRegionsSubsystem->OnActorDeleted(this, true);
 		ColorCorrectRegionsSubsystem = nullptr;
 	}
 	
diff --git a/Engine/Plugins/Experimental/ColorCorrectRegions/Source/ColorCorrectRegions/Private/ColorCorrectRegionSubsystem.cpp b/Engine/Plugins/Experimental/ColorCorrectRegions/Source/ColorCorrectRegions/Private/ColorCorrectRegionSubsystem.cpp
--- a/Engine/Plugins/Experimental/ColorCorrectRegions/Source/ColorCorrectRegions/Private/ColorCorrectRegionSubsystem.cpp
+++ b/Engine/Plugins/Experimental/ColorCorrectRegions/Source/ColorCorrectRegions/Private/ColorCorrectRegionSubsystem.cpp
@@ -61,7 +61,7 @@ void UColorCorrectRegionsSubsystem::Initialize(FSubsystemCollectionBase& Collect
 	if (GetWorld()->WorldType == EWorldType::Editor)
 	{
 		GEngine->OnLevelActorAdded().AddUObject(this, &UColorCorrectRegionsSubsystem::OnActorSpawned);
-		GEngine->OnLevelActorDeleted().AddUObject(this, &UColorCorrectRegionsSubsystem::OnActorDeleted);
+		GEngine->OnLevelActorDeleted().AddUObject(this, &UColorCorrectRegionsSubsystem::OnActorDeleted, true);
 		GEngine->OnLevelActorListChanged().AddUObject(this, &UColorCorrectRegionsSubsystem::OnLevelActorListChanged);
 		GEditor->RegisterForUndo(this);
 
@@ -172,7 +172,7 @@ void UColorCorrectRegionsSubsystem::OnActorSpawned(AActor* InActor)
 	}
 }
 
-void UColorCorrectRegionsSubsystem::OnActorDeleted(AActor* InActor)
+void UColorCorrectRegionsSubsystem::OnActorDeleted(AActor* InActor, bool bClearStencilIdValues)
 {
 	AColorCorrectRegion* AsRegion = Cast<AColorCorrectRegion>(InActor);
 	if (AsRegion 
@@ -183,7 +183,12 @@ void UColorCorrectRegionsSubsystem::OnActorDeleted(AActor* InActor)
 #endif
 	{
 #if WITH_EDITOR
-		FColorCorrectRegionsStencilManager::OnCCRRemoved(GetWorld(), AsRegion);
+		// In some cases, specifically in case when EndPlay is called or when CCA are part of a hidden sublevel
+		// we don't want the stencil Ids to be reset.
+		if (bClearStencilIdValues)
+		{
+			FColorCorrectRegionsStencilManager::OnCCRRemoved(GetWorld(), AsRegion);
+		}
 #endif
 
 		FScopeLock RegionScopeLock(&RegionAccessCriticalSection);
diff --git a/Engine/Plugins/Experimental/ColorCorrectRegions/Source/ColorCorrectRegions/Public/ColorCorrectRegionsSubsystem.h b/Engine/Plugins/Experimental/ColorCorrectRegions/Source/ColorCorrectRegions/Public/ColorCorrectRegionsSubsystem.h
--- a/Engine/Plugins/Experimental/ColorCorrectRegions/Source/ColorCorrectRegions/Public/ColorCorrectRegionsSubsystem.h
+++ b/Engine/Plugins/Experimental/ColorCorrectRegions/Source/ColorCorrectRegions/Public/ColorCorrectRegionsSubsystem.h
@@ -67,7 +67,7 @@ public:
 	void OnActorSpawned(AActor* InActor);
 
 	/** A callback for CC Region deletion. */
-	void OnActorDeleted(AActor* InActor);
+	void OnActorDeleted(AActor* InActor, bool bClearStencilIdValues);
 
 	/** Called when level is added or removed. */
 	void OnLevelsChanged() { RefreshRegions(); };
