VulkanRHI: Get around Vulkan shader compilation issue in bindless where spirv can't recognize the SamplerState type of the array's SafeType.
#rnx

diff --git a/Engine/Plugins/Experimental/ColorCorrectRegions/Shaders/Private/ColorCorrectRegionsCommon.usf b/Engine/Plugins/Experimental/ColorCorrectRegions/Shaders/Private/ColorCorrectRegionsCommon.usf
--- a/Engine/Plugins/Experimental/ColorCorrectRegions/Shaders/Private/ColorCorrectRegionsCommon.usf
+++ b/Engine/Plugins/Experimental/ColorCorrectRegions/Shaders/Private/ColorCorrectRegionsCommon.usf
@@ -288,6 +288,18 @@ MaterialFloat3 RotateTranslateScale(MaterialFloat3 OriginalVec, MaterialFloat3 T
 }
 
 
+#define POST_PROCESS_INPUT(Index)	\
+{									\
+	if (bFiltered)					\
+	{								\
+		return Texture2DSample(PostProcessInput_##Index##_Texture, PostProcessInput_BilinearSampler, UV); \
+	}								\
+	else							\
+	{								\
+		return Texture2DSample(PostProcessInput_##Index##_Texture, PostProcessInput_##Index##_SharedSampler, UV); \
+	}								\
+}
+
 /** Applies an offset to the scene texture lookup and decodes the HDR linear space color. */
 float4 SceneTextureLookup(float2 UV, int SceneTextureIndex, bool bFiltered)
 {
@@ -304,15 +316,15 @@ float4 SceneTextureLookup(float2 UV, int SceneTextureIndex, bool bFiltered)
 		return float4(1, 1, 1, 1);
 
 	case PPI_PostProcessInput0:
-		return Texture2DSample(PostProcessInput_0_Texture, bFiltered ? PostProcessInput_BilinearSampler : PostProcessInput_0_SharedSampler, UV);
+		POST_PROCESS_INPUT(0);
 	case PPI_PostProcessInput1:
-		return Texture2DSample(PostProcessInput_1_Texture, bFiltered ? PostProcessInput_BilinearSampler : PostProcessInput_1_SharedSampler, UV);
+		POST_PROCESS_INPUT(1);
 	case PPI_PostProcessInput2:
-		return Texture2DSample(PostProcessInput_2_Texture, bFiltered ? PostProcessInput_BilinearSampler : PostProcessInput_2_SharedSampler, UV);
+		POST_PROCESS_INPUT(2);
 	case PPI_PostProcessInput3:
-		return Texture2DSample(PostProcessInput_3_Texture, bFiltered ? PostProcessInput_BilinearSampler : PostProcessInput_3_SharedSampler, UV);
+		POST_PROCESS_INPUT(3);
 	case PPI_PostProcessInput4:
-		return Texture2DSample(PostProcessInput_4_Texture, bFiltered ? PostProcessInput_BilinearSampler : PostProcessInput_4_SharedSampler, UV);
+		POST_PROCESS_INPUT(4);
 #if SHADING_PATH_DEFERRED
 	case PPI_SceneDepth:
 		return ScreenSpaceData.GBuffer.Depth;
