Color Correction Actors: Refactor. Color Correction regions no longer react to varios events and instead are gathered on tick and its state is fully copied to avoid being used on render thread.

diff --git a/Engine/Plugins/Experimental/ColorCorrectRegions/Source/ColorCorrectRegions/Private/ColorCorrectRegion.cpp b/Engine/Plugins/Experimental/ColorCorrectRegions/Source/ColorCorrectRegions/Private/ColorCorrectRegion.cpp
--- a/Engine/Plugins/Experimental/ColorCorrectRegions/Source/ColorCorrectRegions/Private/ColorCorrectRegion.cpp
+++ b/Engine/Plugins/Experimental/ColorCorrectRegions/Source/ColorCorrectRegions/Private/ColorCorrectRegion.cpp
@@ -38,8 +38,6 @@ AColorCorrectRegion::AColorCorrectRegion(const FObjectInitializer& ObjectInitial
 	, Enabled(true)
 	, bEnablePerActorCC(false)
 	, PerActorColorCorrection(EColorCorrectRegionStencilType::IncludeStencil)
-	, ColorCorrectRegionsSubsystem(nullptr)
-	, ColorCorrectRenderProxy(MakeShared<FColorCorrectRenderProxy>())
 {
 	PrimaryActorTick.bCanEverTick = true;
 
@@ -52,13 +50,6 @@ AColorCorrectRegion::AColorCorrectRegion(const FObjectInitializer& ObjectInitial
 	IdentityComponent->CastShadow = false;
 	IdentityComponent->SetHiddenInGame(false);
 
-#if WITH_METADATA
-	if (!Cast<AColorCorrectionWindow>(this))
-	{
-		CreateIcon();
-	}
-#endif
-
 #if WITH_EDITOR
 	if (!IsTemplate())
 	{
@@ -68,161 +59,11 @@ AColorCorrectRegion::AColorCorrectRegion(const FObjectInitializer& ObjectInitial
 #endif
 }
 
-void AColorCorrectRegion::BeginPlay()
-{	
-	Super::BeginPlay();
-	if (const UWorld* World = GetWorld())
-	{
-		ColorCorrectRegionsSubsystem = World->GetSubsystem<UColorCorrectRegionsSubsystem>();
-	}
-
-	if (ColorCorrectRegionsSubsystem.IsValid())
-	{
-		ColorCorrectRegionsSubsystem->OnActorSpawned(this);
-	}
-}
-
-void AColorCorrectRegion::EndPlay(const EEndPlayReason::Type EndPlayReason)
-{
-	if (ColorCorrectRegionsSubsystem.IsValid())
-	{
-		ColorCorrectRegionsSubsystem->OnActorDeleted(this, false);
-		ColorCorrectRegionsSubsystem = nullptr;
-	}
-	Super::EndPlay(EndPlayReason);
-}
-
-void AColorCorrectRegion::BeginDestroy()
-{
-	if (ColorCorrectRegionsSubsystem.IsValid())
-	{
-		ColorCorrectRegionsSubsystem->OnActorDeleted(this, true);
-		ColorCorrectRegionsSubsystem = nullptr;
-	}
-	
-	Super::BeginDestroy();
-}
-
 bool AColorCorrectRegion::ShouldTickIfViewportsOnly() const
 {
 	return true;
 }
 
-void AColorCorrectRegion::TickActor(float DeltaTime, ELevelTick TickType, FActorTickFunction& ThisTickFunction)
-{
-	Super::Tick(DeltaTime);
-
-	TRACE_CPUPROFILER_EVENT_SCOPE_TEXT(*FString::Printf(TEXT("CCR.TickActor %s"), *GetName()));
-
-	TransferState();
-
-
-	// Check to make sure that no ids have been changed externally.
-	{
-		TimeWaited += DeltaTime;
-		const float WaitTimeInSecs = 1.0;
-
-		if (!ColorCorrectRegionsSubsystem.IsValid())
-		{
-			if (const UWorld* World = GetWorld())
-			{
-				ColorCorrectRegionsSubsystem = World->GetSubsystem<UColorCorrectRegionsSubsystem>();
-			}
-		}
-		
-		if (ColorCorrectRegionsSubsystem.IsValid() && TimeWaited >= WaitTimeInSecs)
-		{
-			ColorCorrectRegionsSubsystem->CheckAssignedActorsValidity(this);
-			TimeWaited = 0;
-		}
-
-	}
-}
-
-void AColorCorrectRegion::TransferState()
-{
-	TRACE_CPUPROFILER_EVENT_SCOPE_TEXT(*FString::Printf(TEXT("CCR.TransferState %s"), *GetName()));
-
-	FColorCorrectRenderProxyPtr TempCCRStateRenderThread = MakeShared<FColorCorrectRenderProxy>();
-	if (AColorCorrectionWindow* CCWindow = Cast<AColorCorrectionWindow>(this))
-	{
-		TempCCRStateRenderThread->WindowType = CCWindow->WindowType;
-	}
-	else
-	{
-		TempCCRStateRenderThread->Type = Type;
-	}
-
-	TempCCRStateRenderThread->bIsActiveThisFrame = IsValid(this)
-		&& Enabled
-#if WITH_EDITOR
-		&& !IsHiddenEd()
-#endif 
-		&& !(GetWorld()->HasBegunPlay() && IsHidden());
-
-	TempCCRStateRenderThread->World = GetWorld();
-	TempCCRStateRenderThread->Priority = Priority;
-	TempCCRStateRenderThread->Intensity = Intensity;
-
-	// Inner could be larger than outer, in which case we need to make sure these are swapped.
-	TempCCRStateRenderThread->Inner = FMath::Min<float>(Outer, Inner);
-	TempCCRStateRenderThread->Outer = FMath::Max<float>(Outer, Inner);
-
-	if (TempCCRStateRenderThread->Inner == TempCCRStateRenderThread->Outer)
-	{
-		TempCCRStateRenderThread->Inner -= 0.0001;
-	}
-
-	TempCCRStateRenderThread->Falloff = Falloff;
-	TempCCRStateRenderThread->Invert = Invert;
-	TempCCRStateRenderThread->TemperatureType = TemperatureType;
-	TempCCRStateRenderThread->Temperature = Temperature;
-	TempCCRStateRenderThread->Tint = Tint;
-	TempCCRStateRenderThread->ColorGradingSettings = ColorGradingSettings;
-	TempCCRStateRenderThread->bEnablePerActorCC = bEnablePerActorCC;
-	TempCCRStateRenderThread->PerActorColorCorrection = PerActorColorCorrection;
-
-	GetActorBounds(false, TempCCRStateRenderThread->BoxOrigin, TempCCRStateRenderThread->BoxExtent);
-	TempCCRStateRenderThread->ActorLocation = (FVector3f)GetActorLocation();
-	TempCCRStateRenderThread->ActorRotation = (FVector3f)GetActorRotation().Euler();
-	TempCCRStateRenderThread->ActorScale = (FVector3f)GetActorScale();
-
-	// Transfer Stencil Ids.
-	{
-
-		for (const TSoftObjectPtr<AActor>& StencilActor : AffectedActors)
-		{
-			if (!StencilActor.IsValid())
-			{
-				continue;
-			}
-			TArray<UPrimitiveComponent*> PrimitiveComponents;
-			StencilActor->GetComponents<UPrimitiveComponent>(PrimitiveComponents);
-			for (UPrimitiveComponent* PrimitiveComponent : PrimitiveComponents)
-			{
-				if (PrimitiveComponent->bRenderCustomDepth)
-				{
-					TempCCRStateRenderThread->StencilIds.Add(static_cast<uint32>(PrimitiveComponent->CustomDepthStencilValue));
-				}
-			}
-		}
-	}
-
-	// Store component id to be used on render thread.
-	if (!(TempCCRStateRenderThread->FirstPrimitiveId == IdentityComponent->GetPrimitiveSceneId()))
-	{
-		TempCCRStateRenderThread->FirstPrimitiveId = IdentityComponent->GetPrimitiveSceneId();
-	}
-
-	{
-		ENQUEUE_RENDER_COMMAND(CopyCCProxy)([this, CCRStateToCopy = MoveTemp(TempCCRStateRenderThread)](FRHICommandListImmediate& RHICmdList)
-			{
-				ColorCorrectRenderProxy = CCRStateToCopy;
-			}
-		);
-	}
-
-}
 
 #if WITH_EDITOR
 void AColorCorrectRegion::OnSequencerTimeChanged(TWeakPtr<ISequencer> InSequencer)
@@ -235,6 +76,12 @@ void AColorCorrectRegion::OnSequencerTimeChanged(TWeakPtr<ISequencer> InSequence
 
 void AColorCorrectRegion::HandleAffectedActorsPropertyChange(uint32 ActorListChangeType)
 {
+	TWeakObjectPtr<UColorCorrectRegionsSubsystem> ColorCorrectRegionsSubsystem;
+	if (const UWorld* World = GetWorld())
+	{
+		ColorCorrectRegionsSubsystem = World->GetSubsystem<UColorCorrectRegionsSubsystem>();
+	}
+
 	if (ActorListChangeType == EPropertyChangeType::ArrayAdd
 		|| ActorListChangeType == EPropertyChangeType::ValueSet)
 	{
@@ -258,6 +105,7 @@ void AColorCorrectRegion::HandleAffectedActorsPropertyChange(uint32 ActorListCha
 				AffectedActors.FindOrAdd(TSoftObjectPtr<AActor>());
 			}
 		}
+
 		if (ColorCorrectRegionsSubsystem.IsValid())
 		{
 			ColorCorrectRegionsSubsystem->AssignStencilIdsToPerActorCC(this);
@@ -276,7 +124,7 @@ void AColorCorrectRegion::HandleAffectedActorsPropertyChange(uint32 ActorListCha
 }
 
 #if WITH_METADATA
-void AColorCorrectRegion::CreateIcon()
+void AColorCorrectionRegion::CreateIcon()
 {
 	// Create billboard component
 	if (GIsEditor && !IsRunningCommandlet())
@@ -335,14 +183,6 @@ void AColorCorrectRegion::PostEditChangeProperty(struct FPropertyChangedEvent& P
 {
 	const FName PropertyName = PropertyChangedEvent.GetPropertyName();
 
-	if (!ColorCorrectRegionsSubsystem.IsValid())
-	{
-		if (const UWorld* World = GetWorld())
-		{
-			ColorCorrectRegionsSubsystem = World->GetSubsystem<UColorCorrectRegionsSubsystem>();
-		}
-	}
-
 	if (PropertyName == GET_MEMBER_NAME_CHECKED(AColorCorrectRegion, AffectedActors))
 	{
 		/** Since there might be Dialogs involved we need to run this on game thread. */
@@ -352,17 +192,6 @@ void AColorCorrectRegion::PostEditChangeProperty(struct FPropertyChangedEvent& P
 		});
 	}
 
-	// Reorder all CCRs after the Priority property has changed.
-	// Also, in context of Multi-User: PropertyChangedEvent can be a stub without the actual property data. 
-	// Therefore we need to refresh priority if PropertyChangedEvent.Property is nullptr. 
-	if (PropertyName == GET_MEMBER_NAME_CHECKED(AColorCorrectRegion, Priority) || PropertyChangedEvent.Property == nullptr)
-	{
-		if (ColorCorrectRegionsSubsystem.IsValid())
-		{
-			ColorCorrectRegionsSubsystem->SortRegionsByPriority();
-		}
-	}
-
 	// Stage actor properties
 	{
 		const FStructProperty* StructProperty = CastField<FStructProperty>(PropertyChangedEvent.MemberProperty);
@@ -550,6 +379,11 @@ FName AColorCorrectRegion::GetPositionalPropertiesMemberName() const
 AColorCorrectionRegion::AColorCorrectionRegion(const FObjectInitializer& ObjectInitializer)
 	: Super(ObjectInitializer)
 {
+
+#if WITH_METADATA
+		CreateIcon();
+#endif
+
 	UMaterial* Material = LoadObject<UMaterial>(NULL, TEXT("/ColorCorrectRegions/Materials/M_ColorCorrectRegionTransparentPreview.M_ColorCorrectRegionTransparentPreview"), NULL, LOAD_None, NULL);
 	const TArray<UStaticMesh*> StaticMeshes =
 	{
diff --git a/Engine/Plugins/Experimental/ColorCorrectRegions/Source/ColorCorrectRegions/Private/ColorCorrectRegionSubsystem.cpp b/Engine/Plugins/Experimental/ColorCorrectRegions/Source/ColorCorrectRegions/Private/ColorCorrectRegionSubsystem.cpp
--- a/Engine/Plugins/Experimental/ColorCorrectRegions/Source/ColorCorrectRegions/Private/ColorCorrectRegionSubsystem.cpp
+++ b/Engine/Plugins/Experimental/ColorCorrectRegions/Source/ColorCorrectRegions/Private/ColorCorrectRegionSubsystem.cpp
@@ -53,6 +53,85 @@ namespace
 			InRegion->Priority = HighestPriority + (HighestPriority == 0 ? 1 : FMath::Max(CVarCCRPriorityIncrement.GetValueOnAnyThread(), 1));
 		}
 	}
+
+	FColorCorrectRenderProxyPtr CreateRenderStateForCCActor(AColorCorrectRegion* InActor)
+	{
+		TRACE_CPUPROFILER_EVENT_SCOPE_TEXT(*FString::Printf(TEXT("CCR.TransferState %s"), *InActor->GetName()));
+
+		FColorCorrectRenderProxyPtr TempCCRStateRenderThread = MakeShared<FColorCorrectRenderProxy>();
+		if (AColorCorrectionWindow* CCWindow = Cast<AColorCorrectionWindow>(InActor))
+		{
+			TempCCRStateRenderThread->WindowType = CCWindow->WindowType;
+			TempCCRStateRenderThread->ProxyType = FColorCorrectRenderProxy::DistanceBased;
+		}
+		else
+		{
+			TempCCRStateRenderThread->Type = InActor->Type;
+			TempCCRStateRenderThread->ProxyType = FColorCorrectRenderProxy::PriorityBased;
+		}
+
+		TempCCRStateRenderThread->bIsActiveThisFrame = IsValid(InActor)
+			&& InActor->Enabled
+#if WITH_EDITOR
+			&& !InActor->IsHiddenEd()
+#endif 
+			&& !(InActor->GetWorld()->HasBegunPlay() && InActor->IsHidden());
+
+		TempCCRStateRenderThread->World = InActor->GetWorld();
+		TempCCRStateRenderThread->Priority = InActor->Priority;
+		TempCCRStateRenderThread->Intensity = InActor->Intensity;
+
+		// Inner could be larger than outer, in which case we need to make sure these are swapped.
+		TempCCRStateRenderThread->Inner = FMath::Min<float>(InActor->Outer, InActor->Inner);
+		TempCCRStateRenderThread->Outer = FMath::Max<float>(InActor->Outer, InActor->Inner);
+
+		if (TempCCRStateRenderThread->Inner == TempCCRStateRenderThread->Outer)
+		{
+			TempCCRStateRenderThread->Inner -= 0.0001;
+		}
+
+		TempCCRStateRenderThread->Falloff = InActor->Falloff;
+		TempCCRStateRenderThread->Invert = InActor->Invert;
+		TempCCRStateRenderThread->TemperatureType = InActor->TemperatureType;
+		TempCCRStateRenderThread->Temperature = InActor->Temperature;
+		TempCCRStateRenderThread->Tint = InActor->Tint;
+		TempCCRStateRenderThread->ColorGradingSettings = InActor->ColorGradingSettings;
+		TempCCRStateRenderThread->bEnablePerActorCC = InActor->bEnablePerActorCC;
+		TempCCRStateRenderThread->PerActorColorCorrection = InActor->PerActorColorCorrection;
+
+		InActor->GetActorBounds(false, TempCCRStateRenderThread->BoxOrigin, TempCCRStateRenderThread->BoxExtent);
+		TempCCRStateRenderThread->ActorLocation = (FVector3f)InActor->GetActorLocation();
+		TempCCRStateRenderThread->ActorRotation = (FVector3f)InActor->GetActorRotation().Euler();
+		TempCCRStateRenderThread->ActorScale = (FVector3f)InActor->GetActorScale();
+
+		// Transfer Stencil Ids.
+		{
+
+			for (const TSoftObjectPtr<AActor>& StencilActor : InActor->AffectedActors)
+			{
+				if (!StencilActor.IsValid())
+				{
+					continue;
+				}
+				TArray<UPrimitiveComponent*> PrimitiveComponents;
+				StencilActor->GetComponents<UPrimitiveComponent>(PrimitiveComponents);
+				for (UPrimitiveComponent* PrimitiveComponent : PrimitiveComponents)
+				{
+					if (PrimitiveComponent->bRenderCustomDepth)
+					{
+						TempCCRStateRenderThread->StencilIds.Add(static_cast<uint32>(PrimitiveComponent->CustomDepthStencilValue));
+					}
+				}
+			}
+		}
+
+		// Store component id to be used on render thread.
+		if (!(TempCCRStateRenderThread->FirstPrimitiveId == InActor->IdentityComponent->GetPrimitiveSceneId()))
+		{
+			TempCCRStateRenderThread->FirstPrimitiveId = InActor->IdentityComponent->GetPrimitiveSceneId();
+		}
+		return TempCCRStateRenderThread;
+	}
 }
 
 void UColorCorrectRegionsSubsystem::Initialize(FSubsystemCollectionBase& Collection)
@@ -62,8 +141,6 @@ void UColorCorrectRegionsSubsystem::Initialize(FSubsystemCollectionBase& Collect
 	{
 		GEngine->OnLevelActorAdded().AddUObject(this, &UColorCorrectRegionsSubsystem::OnActorSpawned);
 		GEngine->OnLevelActorDeleted().AddUObject(this, &UColorCorrectRegionsSubsystem::OnActorDeleted, true);
-		GEngine->OnLevelActorListChanged().AddUObject(this, &UColorCorrectRegionsSubsystem::OnLevelActorListChanged);
-		GEditor->RegisterForUndo(this);
 
 		FEditorDelegates::OnDuplicateActorsBegin.AddUObject(this, &UColorCorrectRegionsSubsystem::OnDuplicateActorsBegin);
 		FEditorDelegates::OnDuplicateActorsEnd.AddUObject(this, &UColorCorrectRegionsSubsystem::OnDuplicateActorsEnd);
@@ -72,10 +149,10 @@ void UColorCorrectRegionsSubsystem::Initialize(FSubsystemCollectionBase& Collect
 		FEditorDelegates::OnEditPasteActorsEnd.AddUObject(this, &UColorCorrectRegionsSubsystem::OnDuplicateActorsEnd);
 	}
 #endif
-	// In some cases (like nDisplay nodes) EndPlay is not guaranteed to be called when level is removed.
-	GetWorld()->OnLevelsChanged().AddUObject(this, &UColorCorrectRegionsSubsystem::OnLevelsChanged);
+
 	// Initializing Scene view extension responsible for rendering regions.
-	PostProcessSceneViewExtension = FSceneViewExtensions::NewExtension<FColorCorrectRegionsSceneViewExtension>(this);
+	PostProcessSceneViewExtension = FSceneViewExtensions::NewExtension<FColorCorrectRegionsSceneViewExtension>(GetWorld(), this);
+	Super::Initialize(Collection);
 }
 
 void UColorCorrectRegionsSubsystem::Deinitialize()
@@ -85,8 +162,6 @@ void UColorCorrectRegionsSubsystem::Deinitialize()
 	{
 		GEngine->OnLevelActorAdded().RemoveAll(this);
 		GEngine->OnLevelActorDeleted().RemoveAll(this);
-		GEngine->OnLevelActorListChanged().RemoveAll(this);
-		GEditor->UnregisterForUndo(this);
 
 		FEditorDelegates::OnDuplicateActorsBegin.RemoveAll(this);
 		FEditorDelegates::OnDuplicateActorsEnd.RemoveAll(this);
@@ -95,7 +170,6 @@ void UColorCorrectRegionsSubsystem::Deinitialize()
 		FEditorDelegates::OnEditPasteActorsEnd.RemoveAll(this);
 	}
 #endif
-	GetWorld()->OnLevelsChanged().RemoveAll(this);
 
 	// Prevent this SVE from being gathered, in case it is kept alive by a strong reference somewhere else.
 	{
@@ -137,35 +211,68 @@ void UColorCorrectRegionsSubsystem::Deinitialize()
 
 	RegionsPriorityBased.Reset();
 	RegionsDistanceBased.Reset();
+	Super::Deinitialize();
 }
 
-void UColorCorrectRegionsSubsystem::OnActorSpawned(AActor* InActor)
+void UColorCorrectRegionsSubsystem::Tick(float DeltaTime)
 {
-	AColorCorrectRegion* AsRegion = Cast<AColorCorrectRegion>(InActor);
-	if (IsRegionValid(AsRegion, GetWorld()))
+	Super::Tick(DeltaTime);
+
+	RefreshRegions();
+	TransferStates();
+
+	// Check to make sure that no ids have been changed externally.
 	{
-		FScopeLock RegionScopeLock(&RegionAccessCriticalSection);
-		// We wouldn't have to do a check here except in case of nDisplay we need to populate this list during OnLevelsChanged 
-		// because nDisplay can release Actors while those are marked as BeginningPlay. Therefore we want to avoid 
-		// adding regions twice.
-		bool bIsDistanceBased = Cast<AColorCorrectionWindow>(InActor) != nullptr;
-		TArray<AColorCorrectRegion*>* RegionsToAddTo = bIsDistanceBased ? &RegionsDistanceBased : &RegionsPriorityBased;
-		if (!bIsDistanceBased && AsRegion->Priority == 0)
-		{
-			AssignNewPriorityIfNeeded(AsRegion, RegionsPriorityBased);
-		}
-		
-		if (!RegionsToAddTo->Contains(AsRegion))
+		TimeSinceLastValidityCheck += DeltaTime;
+		const float WaitTimeInSecs = 1.0;
+		if (TimeSinceLastValidityCheck >= WaitTimeInSecs)
 		{
-			RegionsToAddTo->Add(AsRegion);
-			// Distance based CCR can only be sorted on render, when View info is available.
-			if (!bIsDistanceBased)
+			for (AColorCorrectRegion* Region : RegionsPriorityBased)
 			{
-				SortRegionsByPriority();
+				CheckAssignedActorsValidity(Region);
 			}
+			for (AColorCorrectRegion* Region : RegionsDistanceBased)
+			{
+				CheckAssignedActorsValidity(Region);
+			}
+
+			TimeSinceLastValidityCheck = 0;
 		}
 	}
+	
+}
 
+void UColorCorrectRegionsSubsystem::TransferStates()
+{
+	TArray<FColorCorrectRenderProxyPtr> TempProxiesPriority;
+	TArray<FColorCorrectRenderProxyPtr> TempProxiesDistance;
+	for (AColorCorrectRegion* Region : RegionsPriorityBased)
+	{
+		TempProxiesPriority.Add(CreateRenderStateForCCActor(Region));
+	}
+
+	for (AColorCorrectRegion* Region : RegionsDistanceBased)
+	{
+		TempProxiesDistance.Add(CreateRenderStateForCCActor(Region));
+	}
+
+	// Sort priority based proxies on game thread.
+	TempProxiesPriority.Sort([](const FColorCorrectRenderProxyPtr& A, const FColorCorrectRenderProxyPtr& B) {
+		// Regions with the same priority could potentially cause flickering on overlap
+		return A->Priority < B->Priority;
+	});
+
+	ENQUEUE_RENDER_COMMAND(CopyCCProxies)([this, TempProxiesPriority = MoveTemp(TempProxiesPriority), TempProxiesDistance = MoveTemp(TempProxiesDistance)](FRHICommandListImmediate& RHICmdList)
+	{
+		ProxiesPriorityBased = TempProxiesPriority;
+		ProxiesDistanceBased = TempProxiesDistance;
+	}
+	);
+}
+
+
+void UColorCorrectRegionsSubsystem::OnActorSpawned(AActor* InActor)
+{
 	if (bDuplicationStarted)
 	{
 		DuplicatedActors.Add(InActor);
@@ -190,10 +297,6 @@ void UColorCorrectRegionsSubsystem::OnActorDeleted(AActor* InActor, bool bClearS
 			FColorCorrectRegionsStencilManager::OnCCRRemoved(GetWorld(), AsRegion);
 		}
 #endif
-
-		FScopeLock RegionScopeLock(&RegionAccessCriticalSection);
-		RegionsPriorityBased.Remove(AsRegion);
-		RegionsDistanceBased.Remove(AsRegion);
 	}
 }
 
@@ -225,40 +328,12 @@ void UColorCorrectRegionsSubsystem::OnDuplicateActorsEnd()
 #if WITH_EDITOR
 	if (GEditor)
 	{
-		this->Modify();
+		this->Modify(false);
 		GEditor->EndTransaction();
 	}
 #endif
 }
 
-void UColorCorrectRegionsSubsystem::SortRegionsByPriority()
-{
-	FScopeLock RegionScopeLock(&RegionAccessCriticalSection);
-
-	RegionsPriorityBased.Sort([](const AColorCorrectRegion& A, const AColorCorrectRegion& B) {
-		// Regions with the same priority could potentially cause flickering on overlap
-		return A.Priority < B.Priority;
-	});
-}
-
-void UColorCorrectRegionsSubsystem::SortRegionsByDistance(const FVector& ViewLocation)
-{
-	FScopeLock RegionScopeLock(&RegionAccessCriticalSection);
-	TMap<AColorCorrectRegion*, double> DistanceMap;
-	for (AColorCorrectRegion* Region : RegionsDistanceBased)
-	{
-		FColorCorrectRenderProxyPtr State = Region->GetCCProxy_RenderThread();
-		FVector CameraToRegionVec = (State->BoxOrigin - ViewLocation);
-		DistanceMap.Add(Region, CameraToRegionVec.Dot(CameraToRegionVec));
-	}
-
-	RegionsDistanceBased.Sort([&DistanceMap](const AColorCorrectRegion& A, const AColorCorrectRegion& B) {
-		// Regions with the same distance could potentially cause flickering on overlap
-		return DistanceMap[&A] > DistanceMap[&B];
-	});
-
-}
-
 void UColorCorrectRegionsSubsystem::AssignStencilIdsToPerActorCC(AColorCorrectRegion* Region, bool bIgnoreUserNotificaion, bool bSoftAssign)
 {
 #if WITH_EDITOR
@@ -272,7 +347,7 @@ void UColorCorrectRegionsSubsystem::AssignStencilIdsToPerActorCC(AColorCorrectRe
 #if WITH_EDITOR
 	if (!bSoftAssign && GEditor)
 	{
-		this->Modify();
+		this->Modify(false);
 		GEditor->EndTransaction();
 	}
 #endif
@@ -288,7 +363,7 @@ void UColorCorrectRegionsSubsystem::ClearStencilIdsToPerActorCC(AColorCorrectReg
 #endif
 
 	FColorCorrectRegionsStencilManager::RemoveStencilNumberForSelectedRegion(GetWorld(), Region);
-	this->Modify();
+	this->Modify(false);
 
 #if WITH_EDITOR
 	if (GEditor)
@@ -322,7 +397,6 @@ void UColorCorrectRegionsSubsystem::RefreshStenciIdAssignmentForAllCCR()
 
 void UColorCorrectRegionsSubsystem::RefreshRegions()
 {
-	FScopeLock RegionScopeLock(&RegionAccessCriticalSection);
 	RegionsPriorityBased.Reset();
 	RegionsDistanceBased.Reset();
 	for (TActorIterator<AColorCorrectRegion> It(GetWorld()); It; ++It)
@@ -330,10 +404,6 @@ void UColorCorrectRegionsSubsystem::RefreshRegions()
 		AColorCorrectRegion* AsRegion = *It;
 		if (IsRegionValid(AsRegion, GetWorld()))
 		{
-			// OnLevelChanged is called when Sublevel is hidden. State needs to be transferred
-			// once per actor on refresh since actors don't tick if sublevel is hidden.
-			AsRegion->TransferState();
-
 			if (!Cast<AColorCorrectionWindow>(AsRegion))
 			{
 				RegionsPriorityBased.Add(AsRegion);
@@ -345,8 +415,6 @@ void UColorCorrectRegionsSubsystem::RefreshRegions()
 		}
 	}
 
-	SortRegionsByPriority();
-
 	RefreshStenciIdAssignmentForAllCCR();
 }
 
diff --git a/Engine/Plugins/Experimental/ColorCorrectRegions/Source/ColorCorrectRegions/Private/ColorCorrectRegionsSceneViewExtension.cpp b/Engine/Plugins/Experimental/ColorCorrectRegions/Source/ColorCorrectRegions/Private/ColorCorrectRegionsSceneViewExtension.cpp
--- a/Engine/Plugins/Experimental/ColorCorrectRegions/Source/ColorCorrectRegions/Private/ColorCorrectRegionsSceneViewExtension.cpp
+++ b/Engine/Plugins/Experimental/ColorCorrectRegions/Source/ColorCorrectRegions/Private/ColorCorrectRegionsSceneViewExtension.cpp
@@ -303,11 +303,29 @@ namespace
 		}
 	}
 
+	void SortRegionProxiesByDistance(TArray<FColorCorrectRenderProxyPtr>& ProxiesDistanceBased, const FVector& ViewLocation)
+	{
+		check(IsInRenderingThread());
+
+		TMap<FColorCorrectRenderProxyPtr, double> DistanceMap;
+		for (FColorCorrectRenderProxyPtr State : ProxiesDistanceBased)
+		{
+			FVector CameraToRegionVec = (State->BoxOrigin - ViewLocation);
+			DistanceMap.Add(State, CameraToRegionVec.Dot(CameraToRegionVec));
+		}
+
+		ProxiesDistanceBased.Sort([&DistanceMap](const FColorCorrectRenderProxyPtr& A, const FColorCorrectRenderProxyPtr& B) {
+			// Regions with the same distance could potentially cause flickering on overlap
+			return DistanceMap[A] > DistanceMap[B];
+		});
+
+	}
+
 	bool RenderRegion
 		( FRDGBuilder& GraphBuilder
 		, const FSceneView& View
 		, const FSceneViewFamily& ViewFamily
-		, AColorCorrectRegion* Region
+		, FColorCorrectRenderProxyPtr RegionState
 		, const FIntRect& PrimaryViewRect
 		, const FScreenPassRenderTarget& SceneColorRenderTarget
 		, const float ScreenPercentage
@@ -321,15 +339,11 @@ namespace
 		SCOPED_GPU_STAT(GraphBuilder.RHICmdList, ColorCorrectRegion);
 		FRHIDepthStencilState* DepthStencilState = FScreenPassPipelineState::FDefaultDepthStencilState::GetRHI();
 
-		FColorCorrectRenderProxyPtr RegionState = Region->GetCCProxy_RenderThread();
-
 		/* If Region is pending for kill, invisible or disabled we don't need to render it.
 		*	If Region's Primitive is not visible in the current view's scene then we don't need to render it either.
 		*	We are checking if the region belongs to the same world as the view.
 		*/
 		if (!RegionState->bIsActiveThisFrame ||
-			Region->IsActorBeingDestroyed() ||
-			RegionState->World != ViewFamily.Scene->GetWorld() ||
 			View.HiddenPrimitives.Contains(RegionState->FirstPrimitiveId) ||
 			(View.ShowOnlyPrimitives.IsSet() && !View.ShowOnlyPrimitives->Contains(RegionState->FirstPrimitiveId))
 			)
@@ -444,11 +458,11 @@ namespace
 		}
 
 		TShaderRef<FColorCorrectGenericPS> PixelShader;
-		if (AColorCorrectionWindow* CCWindow = Cast<AColorCorrectionWindow>(Region))
+		if (RegionState->ProxyType == FColorCorrectRenderProxy::DistanceBased)
 		{
 			PixelShader = GetWindowShader(GlobalShaderMap, RegionState->WindowType, TemperatureType, bIsAdvanced, MergedStencilRenderTarget.IsValid());
 		}
-		else
+		else if (RegionState->ProxyType == FColorCorrectRenderProxy::PriorityBased)
 		{
 			PixelShader = GetRegionShader(GlobalShaderMap, RegionState->Type, TemperatureType, bIsAdvanced, MergedStencilRenderTarget.IsValid());
 		}
@@ -642,8 +656,8 @@ namespace
 
 }
 
-FColorCorrectRegionsSceneViewExtension::FColorCorrectRegionsSceneViewExtension(const FAutoRegister& AutoRegister, UColorCorrectRegionsSubsystem* InWorldSubsystem) :
-	FSceneViewExtensionBase(AutoRegister), WorldSubsystem(InWorldSubsystem)
+FColorCorrectRegionsSceneViewExtension::FColorCorrectRegionsSceneViewExtension(const FAutoRegister& AutoRegister, UWorld* InWorld,  UColorCorrectRegionsSubsystem* InWorldSubsystem) :
+	FWorldSceneViewExtension(AutoRegister, InWorld), WorldSubsystem(InWorldSubsystem)
 {
 }
 
@@ -651,10 +665,8 @@ void FColorCorrectRegionsSceneViewExtension::PrePostProcessPass_RenderThread(FRD
 {
 	if (IsValid(WorldSubsystem))
 	{
-		FScopeLock RegionScopeLock(&WorldSubsystem->RegionAccessCriticalSection);
-
 		// Necessary for when an actor is added or removed from the scene. Also when priority is changed.
-		if ((WorldSubsystem->RegionsPriorityBased.Num() == 0 && WorldSubsystem->RegionsDistanceBased.Num() == 0) || !ViewSupportsRegions(View))
+		if ((WorldSubsystem->ProxiesPriorityBased.Num() == 0 && WorldSubsystem->ProxiesDistanceBased.Num() == 0) || !ViewSupportsRegions(View))
 		{
 			return;
 		}
@@ -719,12 +731,10 @@ void FColorCorrectRegionsSceneViewExtension::PrePostProcessPass_RenderThread(FRD
 		// We don't need to do this per region.
 		FSceneTextureShaderParameters SceneTextures = CreateSceneTextureShaderParameters(GraphBuilder, View, ESceneTextureSetupMode::All);
 
-		WorldSubsystem->SortRegionsByDistance(View.ViewLocation);
 		{
-			FScopeLock RegionScopeLock(&WorldSubsystem->RegionAccessCriticalSection);
-			for (auto It = WorldSubsystem->RegionsPriorityBased.CreateConstIterator(); It; ++It)
+			for (auto It = WorldSubsystem->ProxiesPriorityBased.CreateConstIterator(); It; ++It)
 			{
-				AColorCorrectRegion* Region = *It;
+				FColorCorrectRenderProxyPtr Region = *It;
 				RenderRegion(GraphBuilder
 					, View
 					, ViewFamily
@@ -739,9 +749,10 @@ void FColorCorrectRegionsSceneViewExtension::PrePostProcessPass_RenderThread(FRD
 					, GlobalShaderMap
 					, DefaultBlendState);
 			}
-			for (auto It = WorldSubsystem->RegionsDistanceBased.CreateConstIterator(); It; ++It)
+			SortRegionProxiesByDistance(WorldSubsystem->ProxiesDistanceBased, View.ViewLocation);
+			for (auto It = WorldSubsystem->ProxiesDistanceBased.CreateConstIterator(); It; ++It)
 			{
-				AColorCorrectRegion* Region = *It;
+				FColorCorrectRenderProxyPtr Region = *It;
 				RenderRegion(GraphBuilder
 					, View
 					, ViewFamily
diff --git a/Engine/Plugins/Experimental/ColorCorrectRegions/Source/ColorCorrectRegions/Public/ColorCorrectRegion.h b/Engine/Plugins/Experimental/ColorCorrectRegions/Source/ColorCorrectRegions/Public/ColorCorrectRegion.h
--- a/Engine/Plugins/Experimental/ColorCorrectRegions/Source/ColorCorrectRegions/Public/ColorCorrectRegion.h
+++ b/Engine/Plugins/Experimental/ColorCorrectRegions/Source/ColorCorrectRegions/Public/ColorCorrectRegion.h
@@ -62,8 +62,18 @@ public:
 /** A state to store a copy of CCR properties to be used on render thread. */
 struct FColorCorrectRenderProxy
 {
-	EColorCorrectRegionsType Type;
-	EColorCorrectWindowType WindowType;
+	enum EProxyType
+	{
+		PriorityBased,
+		DistanceBased
+	};
+
+	EProxyType ProxyType;
+
+	union {
+		EColorCorrectRegionsType Type;
+		EColorCorrectWindowType WindowType;
+	};
 
 	/** CCR Properties */
 	int32 Priority;
@@ -103,11 +113,8 @@ typedef TSharedPtr<FColorCorrectRenderProxy, ESPMode::ThreadSafe> FColorCorrectR
 
 /**
  * An instance of Color Correction Region. Used to aggregate all active regions.
- * This actor is aggregated by ColorCorrectRegionsSubsystem which handles:
- *   - Level Loaded, Undo/Redo, Added to level, Removed from level events. 
+ * This actor is aggregated by ColorCorrectRegionsSubsystem on Tick. 
  * AActor class itself is not aware of when it is added/removed, Undo/Redo etc in the Editor. 
- * AColorCorrectRegion reaches out to UColorCorrectRegionsSubsystem when its priority is changed, requesting regions to be sorted 
- * or during BeginPlay/EndPlay to register itself. 
  * More information in ColorCorrectRegionsSubsytem.h
  */
 UCLASS(Blueprintable, NotPlaceable, Abstract)
@@ -212,30 +219,21 @@ public:
 	UPROPERTY()
 	TObjectPtr<UColorCorrectionInvisibleComponent> IdentityComponent;
 
-	/** To handle play in Editor, PIE and Standalone. These methods aggregate objects in play mode similarly to 
-	* Editor methods in FColorCorrectRegionsSubsystem
-	*/
-	virtual void BeginPlay() override;
-	virtual void EndPlay(const EEndPlayReason::Type EndPlayReason) override;
-
-	virtual void BeginDestroy() override;
-
-	virtual void TickActor(float DeltaTime, enum ELevelTick TickType, FActorTickFunction& ThisTickFunction);
 	virtual bool ShouldTickIfViewportsOnly() const;
 
+public:
+
 	/**
-	* Gets a full state for rendering. 
+	* Gets a full state for rendering.
 	*/
-	FColorCorrectRenderProxyPtr GetCCProxy_RenderThread()
-	{
-		check(IsInRenderingThread());
-		return ColorCorrectRenderProxy;
-	};
+	UE_DEPRECATED(5.5, "State management is now done by subsystem.")
+	FColorCorrectRenderProxyPtr GetCCProxy_RenderThread() { return nullptr; };
 
 	/**
 	* Copy state required for rendering to be consumed by Scene view extension.
 	*/
-	void TransferState();
+	UE_DEPRECATED(5.5, "State management is now done by subsystem.")
+	void TransferState() {};
 
 protected:
 
@@ -317,12 +315,6 @@ protected:
 
 private:
 
-#if WITH_METADATA
-	/** Creates an icon for CCR/CCW to be clicked on in Editor. */
-	void CreateIcon();
-
-#endif // WITH_METADATA
-
 	/**
 	* AffectedActors property change could potentially invoke a Dialog Window, which should be displayed on Game Thread.
 	* ActorListChangeType represents EPropertyChangeType
@@ -402,18 +394,6 @@ protected:
 
 	/** Update the transform when a positional setter is called. */
 	bool bNotifyOnParamSetter = true;
-	
-private:
-	TWeakObjectPtr<UColorCorrectRegionsSubsystem> ColorCorrectRegionsSubsystem;
-
-	/** A copy of all properties required by render thread to process this CCR. */
-	FColorCorrectRenderProxyPtr ColorCorrectRenderProxy;
-
-	FCriticalSection StateCopyCriticalSecion;
-
-	// This is for optimization purposes that would let us check assigned actors component's stencil ids ever few once in a while.
-	float TimeWaited = 0;
-
 };
 
 /** 
@@ -440,6 +420,12 @@ protected:
 protected:
 	virtual void ChangeShapeVisibilityForActorType() override;
 
+private:
+#if WITH_METADATA
+	/** Creates an icon for CCR/CCW to be clicked on in Editor. */
+	void CreateIcon();
+
+#endif // WITH_METADATA
 };
 
 
diff --git a/Engine/Plugins/Experimental/ColorCorrectRegions/Source/ColorCorrectRegions/Public/ColorCorrectRegionsSceneViewExtension.h b/Engine/Plugins/Experimental/ColorCorrectRegions/Source/ColorCorrectRegions/Public/ColorCorrectRegionsSceneViewExtension.h
--- a/Engine/Plugins/Experimental/ColorCorrectRegions/Source/ColorCorrectRegions/Public/ColorCorrectRegionsSceneViewExtension.h
+++ b/Engine/Plugins/Experimental/ColorCorrectRegions/Source/ColorCorrectRegions/Public/ColorCorrectRegionsSceneViewExtension.h
@@ -8,10 +8,10 @@ class UColorCorrectRegionsSubsystem;
 class UMaterialInterface;
 class FRDGTexture;
 
-class FColorCorrectRegionsSceneViewExtension : public FSceneViewExtensionBase
+class FColorCorrectRegionsSceneViewExtension : public FWorldSceneViewExtension
 {
 public:
-	FColorCorrectRegionsSceneViewExtension(const FAutoRegister& AutoRegister, UColorCorrectRegionsSubsystem* InWorldSubsystem);
+	FColorCorrectRegionsSceneViewExtension(const FAutoRegister& AutoRegister, UWorld* InWorld, UColorCorrectRegionsSubsystem* InWorldSubsystem);
 	
 	//~ Begin FSceneViewExtensionBase Interface
 	virtual void SetupViewFamily(FSceneViewFamily& InViewFamily) override {};
diff --git a/Engine/Plugins/Experimental/ColorCorrectRegions/Source/ColorCorrectRegions/Public/ColorCorrectRegionsSubsystem.h b/Engine/Plugins/Experimental/ColorCorrectRegions/Source/ColorCorrectRegions/Public/ColorCorrectRegionsSubsystem.h
--- a/Engine/Plugins/Experimental/ColorCorrectRegions/Source/ColorCorrectRegions/Public/ColorCorrectRegionsSubsystem.h
+++ b/Engine/Plugins/Experimental/ColorCorrectRegions/Source/ColorCorrectRegions/Public/ColorCorrectRegionsSubsystem.h
@@ -16,19 +16,6 @@
 
 #include "ColorCorrectRegionsSubsystem.generated.h"
 
-/**
- * Conditional inheritance to allow UColorCorrectRegionsSubsystem to inherit/avoid Editor's Undo/Redo in Editor/Game modes.
- */ 
-#if WITH_EDITOR
-class FColorCorrectRegionsEditorUndoClient : public FEditorUndoClient
-{
-};
-#else
-class FColorCorrectRegionsEditorUndoClient
-{
-};
-#endif
-
 /**
  * World Subsystem responsible for managing AColorCorrectRegion classes in level.
  * This subsystem handles:
@@ -45,7 +32,7 @@ class FColorCorrectRegionsEditorUndoClient
  *		AActor does not have an internal event for when its deleted (EndPlay is the closest we have).
  */
 UCLASS()
-class UColorCorrectRegionsSubsystem : public UWorldSubsystem, public FColorCorrectRegionsEditorUndoClient
+class UColorCorrectRegionsSubsystem : public UTickableWorldSubsystem
 {
 	GENERATED_BODY()
 public:
@@ -54,12 +41,9 @@ public:
 	virtual void Initialize(FSubsystemCollectionBase& Collection) override;
 	virtual void Deinitialize() override;
 
-	// Undo/redo is only supported by editor.
-#if WITH_EDITOR
-	// FEditorUndoClient pure virtual methods.
-	virtual void PostUndo(bool bSuccess) override { RefreshRegions(); };
-	virtual void PostRedo(bool bSuccess) override { RefreshRegions(); }
-#endif
+	virtual bool IsTickableInEditor() const { return true; }
+	virtual void Tick(float DeltaTime) override;
+	virtual TStatId GetStatId() const override { RETURN_QUICK_DECLARE_CYCLE_STAT(UColorCorrectRegionsSubsystem, STATGROUP_Tickables); }
 
 public:
 
@@ -69,26 +53,12 @@ public:
 	/** A callback for CC Region deletion. */
 	void OnActorDeleted(AActor* InActor, bool bClearStencilIdValues);
 
-	/** Called when level is added or removed. */
-	void OnLevelsChanged() { RefreshRegions(); };
-
 	/** Called when duplication process is started in the level. */
 	void OnDuplicateActorsBegin() { bDuplicationStarted = true; };
 
 	/** Called when duplication process is ended in the level. */
 	void OnDuplicateActorsEnd();
 
-#if WITH_EDITOR
-	/** A callback for when the level is loaded. */
-	void OnLevelActorListChanged() { RefreshRegions(); };
-#endif
-	
-	/** Sorts regions based on priority. */
-	void SortRegionsByPriority();
-
-	/** Sorts regions based on distance from the camera. */
-	void SortRegionsByDistance(const FVector& ViewLocation);
-
 	/** Handles Stencil Ids for the selected CCR and corresponding actor. */
 	void AssignStencilIdsToPerActorCC(AColorCorrectRegion* Region, bool bIgnoreUserNotificaion = false, bool bSoftAssign = false);
 	
@@ -102,28 +72,58 @@ public:
 	UFUNCTION(BlueprintCallable, meta = (Category = "Color Correct Regions"))
 	void RefreshStenciIdAssignmentForAllCCR();
 
+public:
+#if WITH_EDITOR
+	/** A callback for when the level is loaded. */
+	UE_DEPRECATED(5.5, "CC Actor aggregation is now done on tick.")
+	void OnLevelActorListChanged() {};
+#endif
+
+	/** Sorts regions based on priority. */
+	UE_DEPRECATED(5.5, "Sorting no longer done externally.")
+	void SortRegionsByPriority() {};
+
+	/** Sorts regions based on distance from the camera. */
+	UE_DEPRECATED(5.5, "Scene View Extension is responsible for sorting due to its access to View information.")
+	void SortRegionsByDistance(const FVector& ViewLocation) {};
+
+	/** Called when level is added or removed. */
+	UE_DEPRECATED(5.5, "State management is now done on tick.")
+	void OnLevelsChanged() {};
+
 private:
 
 	/** Repopulates array of region actors. */
 	void RefreshRegions();
 
-public:
+	/**
+	* Copy states required for rendering to be consumed by Scene view extension to render all active 
+	* CCRs and CCWs.
+	*/
+	void TransferStates();
+
+private:
 
 	/** Stores pointers to ColorCorrectRegion Actors that use priority for sorting. */
 	TArray<AColorCorrectRegion*> RegionsPriorityBased;
-
 	/** Stores pointers to ColorCorrectRegion Actors that are based on distance from camera. */
 	TArray<AColorCorrectRegion*> RegionsDistanceBased;
 
-private:
-	TSharedPtr< class FColorCorrectRegionsSceneViewExtension, ESPMode::ThreadSafe > PostProcessSceneViewExtension;
 
-	FCriticalSection RegionAccessCriticalSection;
+	/** Proxies to be used exclusively on render thread. Copies of the state of CC Actors sorted by priority. */
+	TArray<FColorCorrectRenderProxyPtr> ProxiesPriorityBased;
+	/** Proxies to be used exclusively on render thread. Copies of the state of CC Actors sorted by distance. */
+	TArray<FColorCorrectRenderProxyPtr> ProxiesDistanceBased;
+
+	TSharedPtr< class FColorCorrectRegionsSceneViewExtension, ESPMode::ThreadSafe > PostProcessSceneViewExtension;
 
 	/** This is to handle actor duplication for Per Actor CC. */
 	bool bDuplicationStarted = false;
 	TArray<AActor*> DuplicatedActors;
 
+	// This is for optimization purposes that would let us check assigned actors component's stencil ids every once in a while.
+	float TimeSinceLastValidityCheck = 0;
+
 public:
 	friend class FColorCorrectRegionsSceneViewExtension;
 };
\ No newline at end of file
