Color Correction Regions: Fixed an issue where Per actor CC changes were not transacted over MU.

diff --git a/Engine/Plugins/Experimental/ColorCorrectRegions/Source/ColorCorrectRegions/Private/ColorCorrectRegionSubsystem.cpp b/Engine/Plugins/Experimental/ColorCorrectRegions/Source/ColorCorrectRegions/Private/ColorCorrectRegionSubsystem.cpp
--- a/Engine/Plugins/Experimental/ColorCorrectRegions/Source/ColorCorrectRegions/Private/ColorCorrectRegionSubsystem.cpp
+++ b/Engine/Plugins/Experimental/ColorCorrectRegions/Source/ColorCorrectRegions/Private/ColorCorrectRegionSubsystem.cpp
@@ -51,6 +51,9 @@ namespace
 		if (bAssignNewPriority)
 		{
 			InRegion->Priority = HighestPriority + (HighestPriority == 0 ? 1 : FMath::Max(CVarCCRPriorityIncrement.GetValueOnAnyThread(), 1));
+#if WITH_EDITOR
+			InRegion->Modify();
+#endif
 		}
 	}
 
@@ -347,7 +350,7 @@ void UColorCorrectRegionsSubsystem::AssignStencilIdsToPerActorCC(AColorCorrectRe
 #if WITH_EDITOR
 	if (!bSoftAssign && GEditor)
 	{
-		this->Modify(false);
+		this->Modify();
 		GEditor->EndTransaction();
 	}
 #endif
@@ -363,7 +366,7 @@ void UColorCorrectRegionsSubsystem::ClearStencilIdsToPerActorCC(AColorCorrectReg
 #endif
 
 	FColorCorrectRegionsStencilManager::RemoveStencilNumberForSelectedRegion(GetWorld(), Region);
-	this->Modify(false);
+	this->Modify();
 
 #if WITH_EDITOR
 	if (GEditor)
