Color Correct Actors: A more precise way of dirtying actors and components on duplications.

diff --git a/Engine/Plugins/Experimental/ColorCorrectRegions/Source/ColorCorrectRegions/Private/ColorCorrectRegionSubsystem.cpp b/Engine/Plugins/Experimental/ColorCorrectRegions/Source/ColorCorrectRegions/Private/ColorCorrectRegionSubsystem.cpp
--- a/Engine/Plugins/Experimental/ColorCorrectRegions/Source/ColorCorrectRegions/Private/ColorCorrectRegionSubsystem.cpp
+++ b/Engine/Plugins/Experimental/ColorCorrectRegions/Source/ColorCorrectRegions/Private/ColorCorrectRegionSubsystem.cpp
@@ -13,6 +13,7 @@
 
 #if WITH_EDITOR
 #include "Editor.h"
+#include "ScopedTransaction.h"
 #endif
 
 #define LOCTEXT_NAMESPACE "CCR"
@@ -52,6 +53,7 @@ namespace
 		{
 			InRegion->Priority = HighestPriority + (HighestPriority == 0 ? 1 : FMath::Max(CVarCCRPriorityIncrement.GetValueOnAnyThread(), 1));
 #if WITH_EDITOR
+			const FScopedTransaction Transaction(LOCTEXT("NewPriorityAssigned", "New Priority Assigned to CC Actor."));
 			InRegion->Modify();
 #endif
 		}
@@ -307,13 +309,6 @@ void UColorCorrectRegionsSubsystem::OnDuplicateActorsEnd()
 {
 	bDuplicationStarted = false; 
 
-#if WITH_EDITOR
-	if (GEditor)
-	{
-		GEditor->BeginTransaction(LOCTEXT("PerActorCCActorAssigned", "Per actor CC Actor Assigned"));
-	}
-#endif
-
 	for (AActor* DuplicatedActor : DuplicatedActors)
 	{
 		if (AColorCorrectRegion* AsRegion = Cast<AColorCorrectRegion>(DuplicatedActor))
@@ -327,14 +322,6 @@ void UColorCorrectRegionsSubsystem::OnDuplicateActorsEnd()
 	}
 
 	DuplicatedActors.Empty();
-
-#if WITH_EDITOR
-	if (GEditor)
-	{
-		this->Modify(false);
-		GEditor->EndTransaction();
-	}
-#endif
 }
 
 void UColorCorrectRegionsSubsystem::AssignStencilIdsToPerActorCC(AColorCorrectRegion* Region, bool bIgnoreUserNotificaion, bool bSoftAssign)
diff --git a/Engine/Plugins/Experimental/ColorCorrectRegions/Source/ColorCorrectRegions/Private/ColorCorrectRegionsStencilManager.cpp b/Engine/Plugins/Experimental/ColorCorrectRegions/Source/ColorCorrectRegions/Private/ColorCorrectRegionsStencilManager.cpp
--- a/Engine/Plugins/Experimental/ColorCorrectRegions/Source/ColorCorrectRegions/Private/ColorCorrectRegionsStencilManager.cpp
+++ b/Engine/Plugins/Experimental/ColorCorrectRegions/Source/ColorCorrectRegions/Private/ColorCorrectRegionsStencilManager.cpp
@@ -8,6 +8,10 @@
 #include "EngineUtils.h"
 #include "Engine/World.h"
 
+#if WITH_EDITOR
+#include "Editor.h"
+#include "ScopedTransaction.h"
+#endif
 
 
 static TAutoConsoleVariable<int32> CVarStencilIdRangeMin(
@@ -28,6 +32,8 @@ static TAutoConsoleVariable<bool> CVarEnablePopupStencilIdChangedOnActors(
 	TEXT("Allows CCRs to notify users of change in stencil Ids on actors modified by CCR. \n\
 		By default no notification is displayed and CCR assumes that users changed stencil numbers intentionally.\n"));
 
+#define LOCTEXT_NAMESPACE "CCR"
+
 namespace
 {
 	/*
@@ -130,13 +136,32 @@ namespace
 		{
 			TArray<UPrimitiveComponent*> PrimitiveComponents;
 			ActorToAssignStencilTo->GetComponents<UPrimitiveComponent>(PrimitiveComponents);
-#if WITH_EDITOR
-			ActorToAssignStencilTo->Modify();
-#endif
+
+			bool bNeedsDirtying = false;
+			TArray<UPrimitiveComponent*> PrimitiveComponentsToDirty;
+
 			for (UPrimitiveComponent* PrimitiveComponent : PrimitiveComponents)
 			{
 				// If the stencil id is within our range then it is managed by us, so we should reset this to default.
 				if (PrimitiveComponent->bRenderCustomDepth && PrimitiveComponent->CustomDepthStencilValue >= StencilIdMin)
+				{
+					bNeedsDirtying = true;
+					PrimitiveComponentsToDirty.Add(PrimitiveComponent);
+				}
+			}
+
+			if (bNeedsDirtying
+#if WITH_EDITOR
+				&& GEditor
+#endif
+				)
+			{
+#if WITH_EDITOR
+				const FScopedTransaction Transaction(LOCTEXT("PerActorCCActorCleaned", "Per actor CC Actor Stencil Id Cleared."));
+				ActorToAssignStencilTo->Modify();
+#endif
+
+				for (UPrimitiveComponent* PrimitiveComponent : PrimitiveComponentsToDirty)
 				{
 					PrimitiveComponent->SetRenderCustomDepth(false);
 					PrimitiveComponent->CustomDepthStencilValue = 0;
@@ -440,5 +465,6 @@ void FColorCorrectRegionsStencilManager::CheckAssignedActorsValidity(AColorCorre
 
 		}
 	}
-	
 }
+
+#undef LOCTEXT_NAMESPACE
\ No newline at end of file
