[ColorGrading] [ColorCorrectRegions] Move CCR dependencies to CCR package and enable ColorGrading plugin by default

diff --git a/Engine/Plugins/Editor/ColorGrading/ColorGrading.uplugin b/Engine/Plugins/Editor/ColorGrading/ColorGrading.uplugin
--- a/Engine/Plugins/Editor/ColorGrading/ColorGrading.uplugin
+++ b/Engine/Plugins/Editor/ColorGrading/ColorGrading.uplugin
@@ -10,10 +10,10 @@
 	"DocsURL": "",
 	"MarketplaceURL": "",
 	"SupportURL": "",
-	"EnabledByDefault": false,
+	"EnabledByDefault": true,
 	"CanContainContent": false,
 	"IsBetaVersion": false,
-	"Installed": false,
+	"Hidden": true,
 	"Modules": [
 		{
 			"Name": "ColorGradingEditor",
@@ -22,10 +22,6 @@
 		}
 	],
 	"Plugins": [
-		{
-			"Name": "ColorCorrectRegions",
-			"Enabled": true
-		},
 		{
 			"Name": "ObjectMixer",
 			"Enabled": true
diff --git a/Engine/Plugins/Editor/ColorGrading/Source/ColorGradingEditor/ColorGradingEditor.build.cs b/Engine/Plugins/Editor/ColorGrading/Source/ColorGradingEditor/ColorGradingEditor.build.cs
--- a/Engine/Plugins/Editor/ColorGrading/Source/ColorGradingEditor/ColorGradingEditor.build.cs
+++ b/Engine/Plugins/Editor/ColorGrading/Source/ColorGradingEditor/ColorGradingEditor.build.cs
@@ -20,7 +20,6 @@ public class ColorGradingEditor : ModuleRules
 			new string[] {
 				"ApplicationCore",
 				"AppFramework",
-				"ColorCorrectRegions",
 				"Core",
 				"CoreUObject",
 				"DetailCustomizations",
diff --git a/Engine/Plugins/Editor/ColorGrading/Source/ColorGradingEditor/Private/ColorGradingEditorModule.cpp b/Engine/Plugins/Editor/ColorGrading/Source/ColorGradingEditor/Private/ColorGradingEditorModule.cpp
--- a/Engine/Plugins/Editor/ColorGrading/Source/ColorGradingEditor/Private/ColorGradingEditorModule.cpp
+++ b/Engine/Plugins/Editor/ColorGrading/Source/ColorGradingEditor/Private/ColorGradingEditorModule.cpp
@@ -6,11 +6,8 @@
 #include "ColorGradingEditorDataModel.h"
 #include "ColorGradingMixerObjectFilterRegistry.h"
 #include "DataModelGenerators/ColorGradingDataModelGenerator_PostProcessVolume.h"
-#include "DataModelGenerators/ColorGradingDataModelGenerator_ColorCorrectRegion.h"
 #include "SColorGradingPanel.h"
 
-#include "ColorCorrectRegion.h"
-#include "ColorCorrectWindow.h"
 #include "Engine/PostProcessVolume.h"
 #include "Framework/Docking/LayoutExtender.h"
 #include "LevelEditor.h"
@@ -26,18 +23,11 @@ void FColorGradingEditorModule::StartupModule()
 	FColorGradingEditorDataModel::RegisterColorGradingDataModelGenerator<APostProcessVolume>(
 		FGetDetailsDataModelGenerator::CreateStatic(&FColorGradingDataModelGenerator_PostProcessVolume::MakeInstance));
 
-	FColorGradingEditorDataModel::RegisterColorGradingDataModelGenerator<AColorCorrectRegion>(
-		FGetDetailsDataModelGenerator::CreateStatic(&FColorGradingDataModelGenerator_ColorCorrectRegion::MakeInstance));
-
 	FColorGradingListItem::RegisterColorGradingListItemGenerator<FColorGradingListItemGenerator_PostProcessVolume>();
-	FColorGradingListItem::RegisterColorGradingListItemGenerator<FColorGradingListItemGenerator_ColorCorrectRegion>();
 
 	FColorGradingMixerObjectFilterRegistry::RegisterActorClassToPlace(APostProcessVolume::StaticClass());
-	FColorGradingMixerObjectFilterRegistry::RegisterActorClassToPlace(AColorCorrectionRegion::StaticClass());
-	FColorGradingMixerObjectFilterRegistry::RegisterActorClassToPlace(AColorCorrectionWindow::StaticClass());
 
 	FColorGradingMixerObjectFilterRegistry::RegisterObjectClassToFilter(APostProcessVolume::StaticClass());
-	FColorGradingMixerObjectFilterRegistry::RegisterObjectClassToFilter(AColorCorrectRegion::StaticClass());
 
 	FColorGradingCommands::Register();
 
diff --git a/Engine/Plugins/Editor/ColorGrading/Source/ColorGradingEditor/Private/DataModelGenerators/ColorGradingDataModelGenerator_ColorCorrectRegion.cpp b/Engine/Plugins/Editor/ColorGrading/Source/ColorGradingEditor/Private/DataModelGenerators/ColorGradingDataModelGenerator_ColorCorrectRegion.cpp
deleted file mode 100644
diff --git a/Engine/Plugins/Editor/ColorGrading/Source/ColorGradingEditor/Private/DataModelGenerators/ColorGradingDataModelGenerator_ColorCorrectRegion.h b/Engine/Plugins/Editor/ColorGrading/Source/ColorGradingEditor/Private/DataModelGenerators/ColorGradingDataModelGenerator_ColorCorrectRegion.h
deleted file mode 100644
diff --git a/Engine/Plugins/Editor/ColorGrading/Source/ColorGradingEditor/Private/SColorGradingPanel.cpp b/Engine/Plugins/Editor/ColorGrading/Source/ColorGradingEditor/Private/SColorGradingPanel.cpp
--- a/Engine/Plugins/Editor/ColorGrading/Source/ColorGradingEditor/Private/SColorGradingPanel.cpp
+++ b/Engine/Plugins/Editor/ColorGrading/Source/ColorGradingEditor/Private/SColorGradingPanel.cpp
@@ -11,7 +11,6 @@
 #include "ActorTreeItem.h"
 #include "Algo/Compare.h"
 #include "ComponentTreeItem.h"
-#include "ColorCorrectRegion.h"
 #include "Editor.h"
 #include "Engine/Blueprint.h"
 #include "Engine/Engine.h"
diff --git a/Engine/Plugins/Experimental/ColorCorrectRegions/ColorCorrectRegions.uplugin b/Engine/Plugins/Experimental/ColorCorrectRegions/ColorCorrectRegions.uplugin
--- a/Engine/Plugins/Experimental/ColorCorrectRegions/ColorCorrectRegions.uplugin
+++ b/Engine/Plugins/Experimental/ColorCorrectRegions/ColorCorrectRegions.uplugin
@@ -30,6 +30,10 @@
 		{
 			"Name": "nDisplayModularFeatures",
 			"Enabled": true
+		},
+		{
+			"Name": "ColorGrading",
+			"Enabled": true
 		}
 	]
 }
\ No newline at end of file
diff --git a/Engine/Plugins/Experimental/ColorCorrectRegions/Source/ColorCorrectRegionsEditor/ColorCorrectRegionsEditor.Build.cs b/Engine/Plugins/Experimental/ColorCorrectRegions/Source/ColorCorrectRegionsEditor/ColorCorrectRegionsEditor.Build.cs
--- a/Engine/Plugins/Experimental/ColorCorrectRegions/Source/ColorCorrectRegionsEditor/ColorCorrectRegionsEditor.Build.cs
+++ b/Engine/Plugins/Experimental/ColorCorrectRegions/Source/ColorCorrectRegionsEditor/ColorCorrectRegionsEditor.Build.cs
@@ -18,6 +18,7 @@ public class ColorCorrectRegionsEditor : ModuleRules
 				"Slate",
 				"SlateCore",
 				"ColorCorrectRegions",
+				"ColorGradingEditor",
 				"PlacementMode",
 				"SceneOutliner",
 			}
diff --git a/Engine/Plugins/Experimental/ColorCorrectRegions/Source/ColorCorrectRegionsEditor/Private/ColorCorrectRegionsEditorModule.cpp b/Engine/Plugins/Experimental/ColorCorrectRegions/Source/ColorCorrectRegionsEditor/Private/ColorCorrectRegionsEditorModule.cpp
--- a/Engine/Plugins/Experimental/ColorCorrectRegions/Source/ColorCorrectRegionsEditor/Private/ColorCorrectRegionsEditorModule.cpp
+++ b/Engine/Plugins/Experimental/ColorCorrectRegions/Source/ColorCorrectRegionsEditor/Private/ColorCorrectRegionsEditorModule.cpp
@@ -7,6 +7,9 @@
 #include "ColorCorrectRegionCustomization.h"
 #include "ColorCorrectRegionsStyle.h"
 #include "ColorCorrectWindow.h"
+#include "ColorGradingDataModelGenerator_ColorCorrectRegion.h"
+#include "ColorGradingEditorDataModel.h"
+#include "ColorGradingMixerObjectFilterRegistry.h"
 #include "IPlacementModeModule.h"
 
 #define LOCTEXT_NAMESPACE "FColorCorrectRegionsModule"
@@ -19,6 +22,16 @@ void FColorCorrectRegionsEditorModule::StartupModule()
 	FPropertyEditorModule& PropertyModule = FModuleManager::LoadModuleChecked<FPropertyEditorModule>("PropertyEditor");
 	PropertyModule.RegisterCustomClassLayout(AColorCorrectRegion::StaticClass()->GetFName(), FOnGetDetailCustomizationInstance::CreateStatic(&FColorCorrectWindowDetails::MakeInstance));
 
+	FColorGradingEditorDataModel::RegisterColorGradingDataModelGenerator<AColorCorrectRegion>(
+		FGetDetailsDataModelGenerator::CreateStatic(&FColorGradingDataModelGenerator_ColorCorrectRegion::MakeInstance));
+
+	FColorGradingListItem::RegisterColorGradingListItemGenerator<FColorGradingListItemGenerator_ColorCorrectRegion>();
+
+	FColorGradingMixerObjectFilterRegistry::RegisterActorClassToPlace(AColorCorrectionRegion::StaticClass());
+	FColorGradingMixerObjectFilterRegistry::RegisterActorClassToPlace(AColorCorrectionWindow::StaticClass());
+
+	FColorGradingMixerObjectFilterRegistry::RegisterObjectClassToFilter(AColorCorrectRegion::StaticClass());
+
 	ContextMenu = MakeShared<FColorCorrectionActorContextMenu>();
 	ContextMenu->RegisterContextMenuExtender();
 }
diff --git a/Engine/Plugins/Experimental/ColorCorrectRegions/Source/ColorCorrectRegionsEditor/Private/ColorGradingDataModelGenerator_ColorCorrectRegion.cpp b/Engine/Plugins/Experimental/ColorCorrectRegions/Source/ColorCorrectRegionsEditor/Private/ColorGradingDataModelGenerator_ColorCorrectRegion.cpp
new file mode 100644
--- /dev/null
+++ b/Engine/Plugins/Experimental/ColorCorrectRegions/Source/ColorCorrectRegionsEditor/Private/ColorGradingDataModelGenerator_ColorCorrectRegion.cpp
@@ -0,0 +1,223 @@
+// Copyright Epic Games, Inc. All Rights Reserved.
+
+#include "ColorGradingDataModelGenerator_ColorCorrectRegion.h"
+
+#include "ClassIconFinder.h"
+#include "ColorCorrectRegion.h"
+#include "DetailCategoryBuilder.h"
+#include "DetailLayoutBuilder.h"
+#include "IDetailCustomization.h"
+#include "IDetailTreeNode.h"
+#include "IPropertyRowGenerator.h"
+#include "PropertyHandle.h"
+#include "Widgets/Images/SImage.h"
+#include "Widgets/Input/SButton.h"
+#include "Widgets/Layout/SBox.h"
+#include "Widgets/Text/SInlineEditableTextBlock.h"
+#include "Widgets/Text/STextBlock.h"
+
+#define LOCTEXT_NAMESPACE "FColorGradingDataModelGenerator_ColorCorrectRegion"
+
+TSharedRef<IColorGradingEditorDataModelGenerator> FColorGradingDataModelGenerator_ColorCorrectRegion::MakeInstance()
+{
+	return MakeShareable(new FColorGradingDataModelGenerator_ColorCorrectRegion());
+}
+
+class FColorCorrectRegionCustomization : public IDetailCustomization
+{
+public:
+	virtual void CustomizeDetails(IDetailLayoutBuilder& DetailBuilder) override
+	{
+		TArray<FName> Categories;
+		DetailBuilder.GetCategoryNames(Categories);
+
+		for (const FName& Category : Categories)
+		{
+			if (Category != TEXT("Color Correction"))
+			{
+				DetailBuilder.HideCategory(Category);
+			}
+		}
+
+		// TransformCommon is a custom category that doesn't get returned by GetCategoryNames that also needs to be hidden
+		DetailBuilder.HideCategory(TEXT("TransformCommon"));
+
+		IDetailCategoryBuilder& CCCategoryBuilder = DetailBuilder.EditCategory(TEXT("Color Correction"));
+
+		TArray<TSharedRef<IPropertyHandle>> CCPropertyHandles;
+		CCCategoryBuilder.GetDefaultProperties(CCPropertyHandles);
+
+		IDetailCategoryBuilder& ColorGradingElementsCategory = DetailBuilder.EditCategory(FName("ColorGradingElements"));
+
+		ColorGradingElementsCategory.AddProperty(DetailBuilder.GetProperty((GET_MEMBER_NAME_CHECKED(AColorCorrectionRegion, ColorGradingSettings.Global))));
+		ColorGradingElementsCategory.AddProperty(DetailBuilder.GetProperty((GET_MEMBER_NAME_CHECKED(AColorCorrectionRegion, ColorGradingSettings.Shadows))));
+		ColorGradingElementsCategory.AddProperty(DetailBuilder.GetProperty((GET_MEMBER_NAME_CHECKED(AColorCorrectionRegion, ColorGradingSettings.Midtones))));
+		ColorGradingElementsCategory.AddProperty(DetailBuilder.GetProperty((GET_MEMBER_NAME_CHECKED(AColorCorrectionRegion, ColorGradingSettings.Highlights))));
+	}
+};
+
+void FColorGradingDataModelGenerator_ColorCorrectRegion::Initialize(const TSharedRef<class FColorGradingEditorDataModel>& ColorGradingDataModel, const TSharedRef<IPropertyRowGenerator>& PropertyRowGenerator)
+{
+	PropertyRowGenerator->RegisterInstancedCustomPropertyLayout(AColorCorrectRegion::StaticClass(), FOnGetDetailCustomizationInstance::CreateLambda([]
+	{
+		return MakeShared<FColorCorrectRegionCustomization>();
+	}));
+}
+
+void FColorGradingDataModelGenerator_ColorCorrectRegion::Destroy(const TSharedRef<class FColorGradingEditorDataModel>& ColorGradingDataModel, const TSharedRef<IPropertyRowGenerator>& PropertyRowGenerator)
+{
+	PropertyRowGenerator->UnregisterInstancedCustomPropertyLayout(AColorCorrectRegion::StaticClass());
+}
+
+void FColorGradingDataModelGenerator_ColorCorrectRegion::GenerateDataModel(IPropertyRowGenerator& PropertyRowGenerator, FColorGradingEditorDataModel& OutColorGradingDataModel)
+{
+	TArray<TWeakObjectPtr<AColorCorrectRegion>> SelectedCCRs;
+	const TArray<TWeakObjectPtr<UObject>>& SelectedObjects = PropertyRowGenerator.GetSelectedObjects();
+	for (const TWeakObjectPtr<UObject>& SelectedObject : SelectedObjects)
+	{
+		if (SelectedObject.IsValid() && SelectedObject->IsA<AColorCorrectRegion>())
+		{
+			TWeakObjectPtr<AColorCorrectRegion> SelectedRootActor = CastChecked<AColorCorrectRegion>(SelectedObject.Get());
+			SelectedCCRs.Add(SelectedRootActor);
+		}
+	}
+
+	if (!SelectedCCRs.Num())
+	{
+		return;
+	}
+
+	const TArray<TSharedRef<IDetailTreeNode>>& RootNodes = PropertyRowGenerator.GetRootTreeNodes();
+
+	const TSharedRef<IDetailTreeNode>* ColorGradingElementsPtr = RootNodes.FindByPredicate([](const TSharedRef<IDetailTreeNode>& Node)
+	{
+		return Node->GetNodeName() == TEXT("ColorGradingElements");
+	});
+
+	if (ColorGradingElementsPtr)
+	{
+		const TSharedRef<IDetailTreeNode> ColorGradingElements = *ColorGradingElementsPtr;
+		FColorGradingEditorDataModel::FColorGradingGroup ColorGradingGroup;
+
+		ColorGradingGroup.DetailsViewCategories.Add(TEXT("Color Correction"));
+
+		TArray<TSharedRef<IDetailTreeNode>> ColorGradingPropertyNodes;
+		ColorGradingElements->GetChildren(ColorGradingPropertyNodes);
+
+		for (const TSharedRef<IDetailTreeNode>& PropertyNode : ColorGradingPropertyNodes)
+		{
+			TSharedPtr<IPropertyHandle> PropertyHandle = PropertyNode->CreatePropertyHandle();
+
+			FColorGradingEditorDataModel::FColorGradingElement ColorGradingElement = CreateColorGradingElement(PropertyNode, FText::FromName(PropertyNode->GetNodeName()));
+			ColorGradingGroup.ColorGradingElements.Add(ColorGradingElement);
+		}
+
+		ColorGradingGroup.GroupHeaderWidget = SNew(SHorizontalBox)
+			+ SHorizontalBox::Slot()
+			.AutoWidth()
+			.Padding(FMargin(0, 1, 6, 1))
+			.VAlign(VAlign_Center)
+			[
+				SNew(SBox)
+				.WidthOverride(16)
+				.HeightOverride(16)
+				[
+					SNew(SImage)
+					.ColorAndOpacity(FSlateColor::UseForeground())
+					.Image(FClassIconFinder::FindIconForActor(SelectedCCRs[0]))
+				]
+			]
+
+			+ SHorizontalBox::Slot()
+			.AutoWidth()
+			.VAlign(VAlign_Center)
+			.HAlign(HAlign_Left)
+			[
+				SNew(STextBlock)
+				.Text(MakeAttributeLambda([WeakActor = SelectedCCRs[0]]()
+				{
+					return WeakActor.IsValid()
+						? FText::FromString(WeakActor->GetActorLabel())
+						: FText::GetEmpty();
+				}))
+				.Font(FAppStyle::Get().GetFontStyle("NormalFontBold"))
+			];
+
+		OutColorGradingDataModel.ColorGradingGroups.Add(ColorGradingGroup);
+	}
+}
+
+FColorGradingEditorDataModel::FColorGradingElement FColorGradingDataModelGenerator_ColorCorrectRegion::CreateColorGradingElement(const TSharedRef<IDetailTreeNode>& GroupNode, FText ElementLabel)
+{
+	FColorGradingEditorDataModel::FColorGradingElement ColorGradingElement;
+	ColorGradingElement.DisplayName = ElementLabel;
+
+	TArray<TSharedRef<IDetailTreeNode>> ChildNodes;
+	GroupNode->GetChildren(ChildNodes);
+
+	for (const TSharedRef<IDetailTreeNode>& ChildNode : ChildNodes)
+	{
+		TSharedPtr<IPropertyHandle> PropertyHandle = ChildNode->CreatePropertyHandle();
+		if (PropertyHandle.IsValid() && PropertyHandle->IsValidHandle())
+		{
+			const FString ColorGradingModeString = PropertyHandle->GetProperty()->GetMetaData(TEXT("ColorGradingMode")).ToLower();
+
+			if (!ColorGradingModeString.IsEmpty())
+			{
+				if (ColorGradingModeString.Compare(TEXT("saturation")) == 0)
+				{
+					ColorGradingElement.SaturationPropertyHandle = PropertyHandle;
+				}
+				else if (ColorGradingModeString.Compare(TEXT("contrast")) == 0)
+				{
+					ColorGradingElement.ContrastPropertyHandle = PropertyHandle;
+				}
+				else if (ColorGradingModeString.Compare(TEXT("gamma")) == 0)
+				{
+					ColorGradingElement.GammaPropertyHandle = PropertyHandle;
+				}
+				else if (ColorGradingModeString.Compare(TEXT("gain")) == 0)
+				{
+					ColorGradingElement.GainPropertyHandle = PropertyHandle;
+				}
+				else if (ColorGradingModeString.Compare(TEXT("offset")) == 0)
+				{
+					ColorGradingElement.OffsetPropertyHandle = PropertyHandle;
+				}
+			}
+		}
+	}
+
+	return ColorGradingElement;
+}
+
+bool FColorGradingDataModelGenerator_ColorCorrectRegion::FilterDetailsViewProperties(const TSharedRef<IDetailTreeNode>& InDetailTreeNode)
+{
+	if (InDetailTreeNode->GetNodeType() == EDetailNodeType::Category)
+	{
+		return InDetailTreeNode->GetNodeName() == TEXT("Color Correction");
+	}
+	else
+	{
+		return true;
+	}
+}
+
+TArray<TSubclassOf<AActor>> FColorGradingListItemGenerator_ColorCorrectRegion::GetActorClassesForListItems() const
+{ 
+	return { AColorCorrectionRegion::StaticClass() };
+}
+
+void FColorGradingListItemGenerator_ColorCorrectRegion::GenerateColorGradingListItems(AActor* InActor, TArray<FColorGradingListItemRef>& OutList) const
+{
+	if (AColorCorrectionRegion* ColorCorrectRegion = Cast<AColorCorrectionRegion>(InActor))
+	{
+		FColorGradingListItemRef ListItemRef = MakeShared<FColorGradingListItem>(ColorCorrectRegion);
+		ListItemRef->IsItemEnabled = CREATE_IS_ENABLED_LAMBDA(ColorCorrectRegion, ColorCorrectRegion->Enabled);
+		ListItemRef->OnItemEnabledChanged = CREATE_ON_ENABLED_CHANGED_LAMBDA(ColorCorrectRegion, ColorCorrectRegion->Enabled);
+
+		OutList.Add(ListItemRef);
+	}
+}
+
+#undef LOCTEXT_NAMESPACE
\ No newline at end of file
diff --git a/Engine/Plugins/Experimental/ColorCorrectRegions/Source/ColorCorrectRegionsEditor/Private/ColorGradingDataModelGenerator_ColorCorrectRegion.h b/Engine/Plugins/Experimental/ColorCorrectRegions/Source/ColorCorrectRegionsEditor/Private/ColorGradingDataModelGenerator_ColorCorrectRegion.h
new file mode 100644
--- /dev/null
+++ b/Engine/Plugins/Experimental/ColorCorrectRegions/Source/ColorCorrectRegionsEditor/Private/ColorGradingDataModelGenerator_ColorCorrectRegion.h
@@ -0,0 +1,36 @@
+// Copyright Epic Games, Inc. All Rights Reserved.
+
+#pragma once
+
+#include "CoreMinimal.h"
+#include "ColorGradingEditorDataModel.h"
+
+class IDetailTreeNode;
+
+/** Color Grading Data Model Generator for the AColorCorrectionRegion actor class */
+class FColorGradingDataModelGenerator_ColorCorrectRegion : public IColorGradingEditorDataModelGenerator
+{
+public:
+	static TSharedRef<IColorGradingEditorDataModelGenerator> MakeInstance();
+
+	//~ IColorGradingDataModelGenerator interface
+	virtual void Initialize(const TSharedRef<class FColorGradingEditorDataModel>& ColorGradingDataModel, const TSharedRef<IPropertyRowGenerator>& PropertyRowGenerator) override;
+	virtual void Destroy(const TSharedRef<class FColorGradingEditorDataModel>& ColorGradingDataModel, const TSharedRef<IPropertyRowGenerator>& PropertyRowGenerator) override;
+	virtual void GenerateDataModel(IPropertyRowGenerator& PropertyRowGenerator, FColorGradingEditorDataModel& OutColorGradingDataModel) override;
+	//~ End IColorGradingDataModelGenerator interface
+
+private:
+	/** Creates a new color grading element structure for the specified detail tree node, which is expected to have child color properties with the ColorGradingMode metadata set */
+	FColorGradingEditorDataModel::FColorGradingElement CreateColorGradingElement(const TSharedRef<IDetailTreeNode>& GroupNode, FText ElementLabel);
+
+	bool FilterDetailsViewProperties(const TSharedRef<IDetailTreeNode>& InDetailTreeNode);
+};
+
+/** Color grading list item generator for an AColorCorrectionRegion */
+class FColorGradingListItemGenerator_ColorCorrectRegion : public IColorGradingListItemGenerator
+{
+	//~ IColorGradingListItemGenerator interface
+	virtual TArray<TSubclassOf<AActor>> GetActorClassesForListItems() const override;
+	virtual void GenerateColorGradingListItems(AActor* InActor, TArray<FColorGradingListItemRef>& OutList) const override;
+	//~ End IColorGradingListItemGenerator interface
+};
\ No newline at end of file
