Color Correction Regions: Fixed an issue with CCRs/CCWs not copying color properties on Tick.

diff --git a/Engine/Plugins/Experimental/ColorCorrectRegions/Source/ColorCorrectRegions/Private/ColorCorrectRegionSubsystem.cpp b/Engine/Plugins/Experimental/ColorCorrectRegions/Source/ColorCorrectRegions/Private/ColorCorrectRegionSubsystem.cpp
--- a/Engine/Plugins/Experimental/ColorCorrectRegions/Source/ColorCorrectRegions/Private/ColorCorrectRegionSubsystem.cpp
+++ b/Engine/Plugins/Experimental/ColorCorrectRegions/Source/ColorCorrectRegions/Private/ColorCorrectRegionSubsystem.cpp
@@ -224,7 +224,6 @@ void UColorCorrectRegionsSubsystem::Tick(float DeltaTime)
 	Super::Tick(DeltaTime);
 
 	RefreshRegions();
-	TransferStates();
 
 	// Check to make sure that no ids have been changed externally.
 	{
diff --git a/Engine/Plugins/Experimental/ColorCorrectRegions/Source/ColorCorrectRegions/Private/ColorCorrectRegionsSceneViewExtension.cpp b/Engine/Plugins/Experimental/ColorCorrectRegions/Source/ColorCorrectRegions/Private/ColorCorrectRegionsSceneViewExtension.cpp
--- a/Engine/Plugins/Experimental/ColorCorrectRegions/Source/ColorCorrectRegions/Private/ColorCorrectRegionsSceneViewExtension.cpp
+++ b/Engine/Plugins/Experimental/ColorCorrectRegions/Source/ColorCorrectRegions/Private/ColorCorrectRegionsSceneViewExtension.cpp
@@ -663,6 +663,11 @@ FColorCorrectRegionsSceneViewExtension::FColorCorrectRegionsSceneViewExtension(c
 {
 }
 
+void FColorCorrectRegionsSceneViewExtension::BeginRenderViewFamily(FSceneViewFamily& InViewFamily)
+{
+	WorldSubsystem->TransferStates();
+}
+
 void FColorCorrectRegionsSceneViewExtension::PrePostProcessPass_RenderThread(FRDGBuilder& GraphBuilder, const FSceneView& View, const FPostProcessingInputs& Inputs)
 {
 	if (IsValid(WorldSubsystem))
diff --git a/Engine/Plugins/Experimental/ColorCorrectRegions/Source/ColorCorrectRegions/Public/ColorCorrectRegionsSceneViewExtension.h b/Engine/Plugins/Experimental/ColorCorrectRegions/Source/ColorCorrectRegions/Public/ColorCorrectRegionsSceneViewExtension.h
--- a/Engine/Plugins/Experimental/ColorCorrectRegions/Source/ColorCorrectRegions/Public/ColorCorrectRegionsSceneViewExtension.h
+++ b/Engine/Plugins/Experimental/ColorCorrectRegions/Source/ColorCorrectRegions/Public/ColorCorrectRegionsSceneViewExtension.h
@@ -16,7 +16,7 @@ public:
 	//~ Begin FSceneViewExtensionBase Interface
 	virtual void SetupViewFamily(FSceneViewFamily& InViewFamily) override {};
 	virtual void SetupView(FSceneViewFamily& InViewFamily, FSceneView& InView) override {};
-	virtual void BeginRenderViewFamily(FSceneViewFamily& InViewFamily) override {};
+	virtual void BeginRenderViewFamily(FSceneViewFamily& InViewFamily) override;
 	virtual void PrePostProcessPass_RenderThread(FRDGBuilder& GraphBuilder, const FSceneView& View, const FPostProcessingInputs& Inputs) override;
 	//~ End FSceneViewExtensionBase Interface
 public:
