[ColorCorrectRegions] Reorganize CCR details and add button to open Color Grading panel from nDisplay/CCR classes

diff --git a/Engine/Plugins/Editor/ColorGrading/Source/ColorGradingEditor/ColorGradingEditor.build.cs b/Engine/Plugins/Editor/ColorGrading/Source/ColorGradingEditor/ColorGradingEditor.build.cs
--- a/Engine/Plugins/Editor/ColorGrading/Source/ColorGradingEditor/ColorGradingEditor.build.cs
+++ b/Engine/Plugins/Editor/ColorGrading/Source/ColorGradingEditor/ColorGradingEditor.build.cs
@@ -22,6 +22,7 @@ public class ColorGradingEditor : ModuleRules
 				"AppFramework",
 				"Core",
 				"CoreUObject",
+				"CinematicCamera",
 				"DetailCustomizations",
 				"EditorStyle",
 				"Engine",
diff --git a/Engine/Plugins/Editor/ColorGrading/Source/ColorGradingEditor/Private/ColorGradingEditorModule.cpp b/Engine/Plugins/Editor/ColorGrading/Source/ColorGradingEditor/Private/ColorGradingEditorModule.cpp
--- a/Engine/Plugins/Editor/ColorGrading/Source/ColorGradingEditor/Private/ColorGradingEditorModule.cpp
+++ b/Engine/Plugins/Editor/ColorGrading/Source/ColorGradingEditor/Private/ColorGradingEditorModule.cpp
@@ -10,9 +10,11 @@
 #include "SColorGradingPanel.h"
 
 #include "Camera/CameraActor.h"
+#include "CineCameraActor.h"
 #include "Engine/PostProcessVolume.h"
 #include "Framework/Docking/LayoutExtender.h"
 #include "LevelEditor.h"
+#include "PropertyEditorModule.h"
 #include "WorkspaceMenuStructure.h"
 #include "WorkspaceMenuStructureModule.h"
 
@@ -29,6 +31,7 @@ void FColorGradingEditorModule::StartupModule()
 		FGetDetailsDataModelGenerator::CreateStatic(&FColorGradingDataModelGenerator_CameraActor::MakeInstance));
 
 	FColorGradingMixerObjectFilterRegistry::RegisterActorClassToPlace(APostProcessVolume::StaticClass());
+	FColorGradingMixerObjectFilterRegistry::RegisterActorClassToPlace(ACineCameraActor::StaticClass());
 	FColorGradingMixerObjectFilterRegistry::RegisterActorClassToPlace(ACameraActor::StaticClass());
 
 	FColorGradingMixerObjectFilterRegistry::RegisterObjectClassToFilter(APostProcessVolume::StaticClass());
diff --git a/Engine/Plugins/Editor/ColorGrading/Source/ColorGradingEditor/Private/ColorGradingEditorModule.h b/Engine/Plugins/Editor/ColorGrading/Source/ColorGradingEditor/Private/ColorGradingEditorModule.h
--- a/Engine/Plugins/Editor/ColorGrading/Source/ColorGradingEditor/Private/ColorGradingEditorModule.h
+++ b/Engine/Plugins/Editor/ColorGrading/Source/ColorGradingEditor/Private/ColorGradingEditorModule.h
@@ -16,10 +16,11 @@ class FLayoutExtender;
 class FColorGradingEditorModule : public IColorGradingEditor
 {
 public:
-	//~ IModuleInterface interface
+	//~ IColorGradingEditor interface
 	virtual void StartupModule() override;
 	virtual void ShutdownModule() override;
-	//~ End IModuleInterface interface
+	virtual FName GetColorGradingTabSpawnerId() const override { return ColorGradingPanelTabId; }
+	//~ End IColorGradingEditor interface
 
 private:
 	/** Called right before the engine starts ticking */
diff --git a/Engine/Plugins/Editor/ColorGrading/Source/ColorGradingEditor/Private/ColorGradingEditorUtil.cpp b/Engine/Plugins/Editor/ColorGrading/Source/ColorGradingEditor/Private/ColorGradingEditorUtil.cpp
new file mode 100644
--- /dev/null
+++ b/Engine/Plugins/Editor/ColorGrading/Source/ColorGradingEditor/Private/ColorGradingEditorUtil.cpp
@@ -0,0 +1,64 @@
+// Copyright Epic Games, Inc. All Rights Reserved.
+
+#include "ColorGradingEditorUtil.h"
+
+#include "ColorGradingEditorUtil.h"
+#include "ColorGradingEditorStyle.h"
+#include "Framework/Docking/TabManager.h"
+#include "IColorGradingEditor.h"
+#include "LevelEditor.h"
+#include "Styling/CoreStyle.h"
+#include "Widgets/SBoxPanel.h"
+#include "Widgets/Images/SImage.h"
+#include "Widgets/Input/SButton.h"
+
+namespace ColorGradingEditorUtil
+{
+
+TSharedRef<SWidget> MakeColorGradingLaunchButton(bool bWrapInBox)
+{
+	TSharedRef<SWidget> Result = SNew(SButton)
+	.VAlign(VAlign_Center)
+	.HAlign(HAlign_Center)
+	.OnClicked_Lambda([]()
+	{
+		FLevelEditorModule& LevelEditor = FModuleManager::Get().LoadModuleChecked<FLevelEditorModule>("LevelEditor");
+		LevelEditor.GetLevelEditorTabManager()->TryInvokeTab(IColorGradingEditor::Get().GetColorGradingTabSpawnerId());
+		return FReply::Handled();
+	})
+	[
+		SNew(SHorizontalBox)
+
+		+ SHorizontalBox::Slot()
+		.AutoWidth()
+		[
+			SNew(SImage)
+			.Image(FColorGradingEditorStyle::Get().GetBrush("ColorGrading.ToolbarButton"))
+		]
+
+		+ SHorizontalBox::Slot()
+		.AutoWidth()
+		.Padding(4, 0, 0, 0)
+		[
+			SNew(STextBlock)
+			.Visibility(EVisibility::HitTestInvisible)
+			.TextStyle(&FCoreStyle::Get().GetWidgetStyle<FTextBlockStyle>("ButtonText"))
+			.Text(NSLOCTEXT("ColorCorrectWindowDetails", "OpenColorGrading", "Open Color Grading"))
+
+		]
+	];
+
+	if (bWrapInBox)
+	{
+		Result = SNew(SBox)
+		.HAlign(EHorizontalAlignment::HAlign_Center)
+		.Padding(0, 2)
+		[
+			Result
+		];
+	}
+
+	return Result;
+}
+
+}
diff --git a/Engine/Plugins/Editor/ColorGrading/Source/ColorGradingEditor/Public/ColorGradingEditorUtil.h b/Engine/Plugins/Editor/ColorGrading/Source/ColorGradingEditor/Public/ColorGradingEditorUtil.h
new file mode 100644
--- /dev/null
+++ b/Engine/Plugins/Editor/ColorGrading/Source/ColorGradingEditor/Public/ColorGradingEditorUtil.h
@@ -0,0 +1,17 @@
+// Copyright Epic Games, Inc. All Rights Reserved.
+
+#pragma once
+
+#include "Templates/SharedPointer.h"
+#include "Widgets/SWidget.h"
+
+namespace ColorGradingEditorUtil
+{
+
+/**
+ * Make a button that opens the Color Grading editor.
+ * @param bWrapInBox If true, wrap the button in a box that centers it and adds standard padding for use in the Details panel.
+ */
+COLORGRADINGEDITOR_API TSharedRef<SWidget> MakeColorGradingLaunchButton(bool bWrapInBox = true);
+
+}
diff --git a/Engine/Plugins/Editor/ColorGrading/Source/ColorGradingEditor/Public/IColorGradingEditor.h b/Engine/Plugins/Editor/ColorGrading/Source/ColorGradingEditor/Public/IColorGradingEditor.h
--- a/Engine/Plugins/Editor/ColorGrading/Source/ColorGradingEditor/Public/IColorGradingEditor.h
+++ b/Engine/Plugins/Editor/ColorGrading/Source/ColorGradingEditor/Public/IColorGradingEditor.h
@@ -17,6 +17,9 @@ public:
 public:
 	virtual ~IColorGradingEditor() = default;
 
+	/** Get the ID of the main color grading tab */
+	virtual FName GetColorGradingTabSpawnerId() const = 0;
+
 	/**
 	* Singleton-like access to this module's interface.  This is just for convenience!
 	* Beware of calling this during the shutdown phase, though.  Your module might have been unloaded already.
diff --git a/Engine/Plugins/Experimental/ColorCorrectRegions/Source/ColorCorrectRegions/Public/ColorCorrectRegion.h b/Engine/Plugins/Experimental/ColorCorrectRegions/Source/ColorCorrectRegions/Public/ColorCorrectRegion.h
--- a/Engine/Plugins/Experimental/ColorCorrectRegions/Source/ColorCorrectRegions/Public/ColorCorrectRegion.h
+++ b/Engine/Plugins/Experimental/ColorCorrectRegions/Source/ColorCorrectRegions/Public/ColorCorrectRegion.h
@@ -117,7 +117,7 @@ typedef TSharedPtr<FColorCorrectRenderProxy, ESPMode::ThreadSafe> FColorCorrectR
  * AActor class itself is not aware of when it is added/removed, Undo/Redo etc in the Editor. 
  * More information in ColorCorrectRegionsSubsytem.h
  */
-UCLASS(Blueprintable, NotPlaceable, Abstract)
+UCLASS(Blueprintable, NotPlaceable, Abstract, HideCategories=(Actor, HLOD, Collision, Replication, Cooking, Input, DataLayers, Physics, Streaming, WorldPartition))
 class COLORCORRECTREGIONS_API AColorCorrectRegion : public AActor, public IDisplayClusterStageActor
 {
 	GENERATED_UCLASS_BODY()
@@ -129,7 +129,7 @@ public:
 	virtual ~AColorCorrectRegion() override;
 	
 	/** Region type. */
-	UPROPERTY(BlueprintReadWrite, EditAnywhere, Category="Color Correction")
+	UPROPERTY(BlueprintReadWrite, EditAnywhere, Category="Region")
 	EColorCorrectRegionsType Type;
 
 	/** 
@@ -137,55 +137,55 @@ public:
 	* A region with Priority 1 will be rendered before a region with Priority 10. 
 	* This property is hidden if priority is determined by distance from the camera (When Window CCR is being used). 
 	*/
-	UPROPERTY(BlueprintReadWrite, EditAnywhere, Category="Color Correction")
+	UPROPERTY(BlueprintReadWrite, EditAnywhere, Category="Region")
 	int32 Priority;
 
 	/** Color correction intensity. Clamped to 0-1 range. */
-	UPROPERTY(BlueprintReadWrite, EditAnywhere, Category = "Color Correction", meta = (UIMin = 0.0, UIMax = 1.0))
+	UPROPERTY(BlueprintReadWrite, EditAnywhere, Category = "Region", meta = (UIMin = 0.0, UIMax = 1.0))
 	float Intensity;
 
 	/** Inner of the region. Swapped with Outer in case it is higher than Outer. */
-	UPROPERTY(BlueprintReadWrite, EditAnywhere, Category = "Color Correction", meta = (UIMin = 0.0, UIMax = 1.0))
+	UPROPERTY(BlueprintReadWrite, EditAnywhere, Category = "Region", meta = (UIMin = 0.0, UIMax = 1.0))
 	float Inner;
 
 	/** Outer of the region. */
-	UPROPERTY(BlueprintReadWrite, EditAnywhere, Category = "Color Correction", meta = (UIMin = 0.0, UIMax = 1.0))
+	UPROPERTY(BlueprintReadWrite, EditAnywhere, Category = "Region", meta = (UIMin = 0.0, UIMax = 1.0))
 	float Outer;
 
 	/** Falloff. Softening the region. */
-	UPROPERTY(BlueprintReadWrite, EditAnywhere, Category = "Color Correction", meta = (UIMin = 0.0))
+	UPROPERTY(BlueprintReadWrite, EditAnywhere, Category = "Region", meta = (UIMin = 0.0))
 	float Falloff;
 
 	/** Invert region. */
-	UPROPERTY(BlueprintReadWrite, EditAnywhere, Category = "Color Correction")
+	UPROPERTY(BlueprintReadWrite, EditAnywhere, Category = "Region")
 	bool Invert;
 
 	/** Type of algorithm to be used to control color temperature or white balance. */
-	UPROPERTY(BlueprintReadWrite, EditAnywhere, Category = "Color Correction")
+	UPROPERTY(BlueprintReadWrite, EditAnywhere, Category = "Color Grading")
 	EColorCorrectRegionTemperatureType TemperatureType;
 
 	/** Color correction temperature. */
-	UPROPERTY(BlueprintReadWrite, EditAnywhere, Category = "Color Correction", meta = (UIMin = "1500.0", UIMax = "15000.0"))
+	UPROPERTY(BlueprintReadWrite, EditAnywhere, Category = "Color Grading", meta = (UIMin = "1500.0", UIMax = "15000.0"))
 	float Temperature;
 
 	/** Color temperature tint. */
-	UPROPERTY(BlueprintReadWrite, EditAnywhere, Category = "Color Correction", meta=(UIMin = "-1.0", UIMax = "1.0"))
+	UPROPERTY(BlueprintReadWrite, EditAnywhere, Category = "Color Grading", meta=(UIMin = "-1.0", UIMax = "1.0"))
 	float Tint;
 
 	/** Color correction settings. */
-	UPROPERTY(BlueprintReadWrite, EditAnywhere, Category = "Color Correction")
+	UPROPERTY(BlueprintReadWrite, EditAnywhere, Category = "Color Grading", meta = (ShowPostProcessCategories))
 	FColorGradingSettings ColorGradingSettings;
 
 	/** Enable/Disable color correction provided by this region. */
-	UPROPERTY(BlueprintReadWrite, EditAnywhere, Category = "Color Correction")
+	UPROPERTY(BlueprintReadWrite, EditAnywhere, Category = "Color Grading", meta = (DisplayName = "Enable Color Grading", DisplayPriority = 1))
 	bool Enabled;
 
 	/** Enables or disabled per actor color correction. */
-	UPROPERTY(BlueprintReadWrite, EditAnywhere, Category = "Per Actor CC")
+	UPROPERTY(BlueprintReadWrite, EditAnywhere, Category = "Per Actor CC", meta = (DisplayName = "Enable Per-Actor CC"))
 	bool bEnablePerActorCC;
 
 	/** Controls in which way the below targets will be affected by color correction. */
-	UPROPERTY(BlueprintReadWrite, EditAnywhere, Category = "Per Actor CC", meta = (editcondition = "bEnablePerActorCC", DisplayName = "Per Actor CC Mode"))
+	UPROPERTY(BlueprintReadWrite, EditAnywhere, Category = "Per Actor CC", meta = (editcondition = "bEnablePerActorCC", DisplayName = "Per-Actor CC Mode"))
 	EColorCorrectRegionStencilType PerActorColorCorrection;
 
 	/** List of actors that get affected or ignored by Per actor CC. Effect depends on the above option. */
diff --git a/Engine/Plugins/Experimental/ColorCorrectRegions/Source/ColorCorrectRegions/Public/ColorCorrectWindow.h b/Engine/Plugins/Experimental/ColorCorrectRegions/Source/ColorCorrectRegions/Public/ColorCorrectWindow.h
--- a/Engine/Plugins/Experimental/ColorCorrectRegions/Source/ColorCorrectRegions/Public/ColorCorrectWindow.h
+++ b/Engine/Plugins/Experimental/ColorCorrectRegions/Source/ColorCorrectRegions/Public/ColorCorrectWindow.h
@@ -27,7 +27,7 @@ public:
 	virtual ~AColorCorrectionWindow() override;
 
 	/** Region type. */
-	UPROPERTY(BlueprintReadWrite, EditAnywhere, Category="Color Correction", Meta = (DisplayName = "Type"))
+	UPROPERTY(BlueprintReadWrite, EditAnywhere, Category="Region", Meta = (DisplayName = "Type"))
 	EColorCorrectWindowType WindowType;
 
 #if WITH_EDITOR
diff --git a/Engine/Plugins/Experimental/ColorCorrectRegions/Source/ColorCorrectRegionsEditor/Private/ColorCorrectRegionCustomization.cpp b/Engine/Plugins/Experimental/ColorCorrectRegions/Source/ColorCorrectRegionsEditor/Private/ColorCorrectRegionCustomization.cpp
--- a/Engine/Plugins/Experimental/ColorCorrectRegions/Source/ColorCorrectRegionsEditor/Private/ColorCorrectRegionCustomization.cpp
+++ b/Engine/Plugins/Experimental/ColorCorrectRegions/Source/ColorCorrectRegionsEditor/Private/ColorCorrectRegionCustomization.cpp
@@ -3,7 +3,15 @@
 #include "ColorCorrectRegionCustomization.h"
 #include "ColorCorrectRegion.h"
 #include "ColorCorrectWindow.h"
+
+#include "ColorGradingEditorUtil.h"
+#include "DetailCategoryBuilder.h"
 #include "DetailLayoutBuilder.h"
+#include "DetailWidgetRow.h"
+#include "IDetailGroup.h"
+#include "ObjectEditorUtils.h"
+
+#define LOCTEXT_NAMESPACE "ColorCorrectWindowDetails"
 
 TSharedRef<IDetailCustomization> FColorCorrectWindowDetails::MakeInstance()
 {
@@ -12,15 +20,165 @@ TSharedRef<IDetailCustomization> FColorCorrectWindowDetails::MakeInstance()
 
 void FColorCorrectWindowDetails::CustomizeDetails(IDetailLayoutBuilder& DetailBuilder)
 {
-	for (const TWeakObjectPtr<UObject>& SelectedObject : DetailBuilder.GetSelectedObjects())
+	// Add hyphen to Per-Actor CC category. We have to do this here since by default, the editor will add a space after the hyphen
+	const FText PerActorCCCategoryName = LOCTEXT("PerActorCCCategory", "Per-Actor CC");
+	IDetailCategoryBuilder& PerActorCCCategory = DetailBuilder.EditCategory("Per Actor CC", PerActorCCCategoryName);
+	PerActorCCCategory.SetDisplayName(PerActorCCCategoryName);
+
+	// Hide CCR-specific properties if CCWs are present in the selection
+	const bool bHasCCWs = DetailBuilder.GetSelectedObjects().ContainsByPredicate([](const TWeakObjectPtr<UObject>& SelectedObject)
+	{
+		return SelectedObject.IsValid() && SelectedObject->IsA<AColorCorrectionWindow>();
+	});
+
+	if (bHasCCWs)
+	{
+		TSharedRef<IPropertyHandle> PriorityProperty = DetailBuilder.GetProperty(GET_MEMBER_NAME_CHECKED(AColorCorrectRegion, Priority));
+		DetailBuilder.HideProperty(PriorityProperty);
+
+		TSharedRef<IPropertyHandle> TypeProperty = DetailBuilder.GetProperty(GET_MEMBER_NAME_CHECKED(AColorCorrectRegion, Type));
+		DetailBuilder.HideProperty(TypeProperty);
+	}
+
+	// Create custom color grading category with button at top
+	IDetailCategoryBuilder& ColorGradingCategory = DetailBuilder.EditCategory("Color Grading", LOCTEXT("ColorGradingCategory", "Color Grading"));
+	ColorGradingCategory.AddCustomRow(LOCTEXT("OpenColorGrading", "Open Color Grading"))
+		.RowTag("OpenColorGrading")
+		[
+			ColorGradingEditorUtil::MakeColorGradingLaunchButton()
+		];
+
+	// Place Enabled setting immediately under button
+	ColorGradingCategory.AddProperty(DetailBuilder.GetProperty(GET_MEMBER_NAME_CHECKED(AColorCorrectRegion, Enabled)));
+
+	// Move temperature settings to custom group in custom Color Grading category
+	{
+		IDetailGroup& TemperatureGroup = ColorGradingCategory.AddGroup("Temperature", LOCTEXT("ColorGradingTemperatureGroup", "Temperature"));
+
+		const TArray<FName> TemperaturePropertyNames = {
+			GET_MEMBER_NAME_CHECKED(AColorCorrectionRegion, TemperatureType),
+			GET_MEMBER_NAME_CHECKED(AColorCorrectionRegion, Temperature),
+			GET_MEMBER_NAME_CHECKED(AColorCorrectionRegion, Tint)
+		};
+
+		for (const FName& PropertyName : TemperaturePropertyNames)
+		{
+			TSharedPtr<IPropertyHandle> ChildHandle = DetailBuilder.GetProperty(PropertyName);
+			TemperatureGroup.AddPropertyRow(ChildHandle.ToSharedRef());
+		}
+	}
+
+	// Add all Color Grading struct properties directly to the custom category
+	TSharedRef<IPropertyHandle> ColorGradingPropertyHandle = DetailBuilder.GetProperty(GET_MEMBER_NAME_CHECKED(AColorCorrectionRegion, ColorGradingSettings));
+	if (ColorGradingPropertyHandle.Get().IsValidHandle())
+	{
+		ColorGradingPropertyHandle->MarkHiddenByCustomization();
+		MoveColorGradingPropertiesToCategory(ColorGradingPropertyHandle, ColorGradingCategory);
+	}
+
+	// Sort categories
+	DetailBuilder.SortCategories([&](const TMap<FName, IDetailCategoryBuilder*>& CategoryMap)
+	{
+		for (const TPair<FName, IDetailCategoryBuilder*>& Pair : CategoryMap)
+		{
+			int32 SortOrder = Pair.Value->GetSortOrder();
+			const FName CategoryName = Pair.Key;
+			const IDetailCategoryBuilder* Category = Pair.Value;
+
+			if (CategoryName == "TransformCommon")
+			{
+				SortOrder = 0;
+			}
+			else if (CategoryName == "Region")
+			{
+				SortOrder = 1;
+			}
+			else if (Pair.Value == &ColorGradingCategory)
+			{
+				SortOrder = 2;
+			}
+			else if (Pair.Value == &PerActorCCCategory)
+			{
+				SortOrder = 3;
+			}
+			else if (CategoryName == "Orientation")
+			{
+				SortOrder = 4;
+			}
+			else
+			{
+				const int32 ValueSortOrder = Pair.Value->GetSortOrder();
+				if (ValueSortOrder >= SortOrder && ValueSortOrder < SortOrder + 10)
+				{
+					SortOrder += 10;
+				}
+				else
+				{
+					continue;
+				}
+			}
+
+			Pair.Value->SetSortOrder(SortOrder);
+		}
+	});
+}
+
+void FColorCorrectWindowDetails::MoveColorGradingPropertiesToCategory(TSharedRef<IPropertyHandle> StructHandle, IDetailCategoryBuilder& RootCategory)
+{
+	// Map from color grading subproperties to the groups we want to assign them to
+	const TMap<FName, FName> ColorGradingGroupAssignments = {
+		{ GET_MEMBER_NAME_CHECKED(FColorGradingSettings, ShadowsMax), GET_MEMBER_NAME_CHECKED(FColorGradingSettings, Shadows) },
+		{ GET_MEMBER_NAME_CHECKED(FColorGradingSettings, HighlightsMin), GET_MEMBER_NAME_CHECKED(FColorGradingSettings, Highlights) },
+		{ GET_MEMBER_NAME_CHECKED(FColorGradingSettings, HighlightsMax), GET_MEMBER_NAME_CHECKED(FColorGradingSettings, Highlights) },
+	};
+
+	TMap<FName, IDetailGroup*> Groups;
+
+	uint32 NumColorGradingChildren = 0;
+	StructHandle->GetNumChildren(NumColorGradingChildren);
+	for (uint32 ColorGradingChildIndex = 0; ColorGradingChildIndex < NumColorGradingChildren; ++ColorGradingChildIndex)
 	{
-		if (AColorCorrectionWindow* CCW = Cast<AColorCorrectionWindow>(SelectedObject.Get()))
+		TSharedPtr<IPropertyHandle> SubpropertyHandle = StructHandle->GetChildHandle(ColorGradingChildIndex);
+		FProperty* Subproperty = SubpropertyHandle->GetProperty();
+		const FName SubpropertyName = Subproperty->GetFName();
+
+		uint32 NumGroupChildren = 0;
+		SubpropertyHandle->GetNumChildren(NumGroupChildren);
+		if (NumGroupChildren > 0)
 		{
-			TSharedRef<IPropertyHandle> PriorityProperty = DetailBuilder.GetProperty(GET_MEMBER_NAME_CHECKED(AColorCorrectRegion, Priority));
-			DetailBuilder.HideProperty(PriorityProperty);
+			// Recreate the property as a group within the Color Grading category
+			IDetailGroup& NewGroup = RootCategory.AddGroup(SubpropertyName, Subproperty->GetDisplayNameText());
+			NewGroup.EnableReset(true);
 
-			TSharedRef<IPropertyHandle> TypeProperty = DetailBuilder.GetProperty(GET_MEMBER_NAME_CHECKED(AColorCorrectRegion, Type));
-			DetailBuilder.HideProperty(TypeProperty);
+			Groups.Emplace(SubpropertyName, &NewGroup);
+
+			for (uint32 GroupChildIndex = 0; GroupChildIndex < NumGroupChildren; ++GroupChildIndex)
+			{
+				TSharedPtr<IPropertyHandle> GroupChildHandle = SubpropertyHandle->GetChildHandle(GroupChildIndex);
+				NewGroup.AddPropertyRow(GroupChildHandle.ToSharedRef());
+			}
+		}
+		else
+		{
+			// Try to assign the property to one of the color grading groups
+			bool bHasAssignedGroup = false;
+
+			if (const FName* AssignedGroupName = ColorGradingGroupAssignments.Find(SubpropertyName))
+			{
+				if (IDetailGroup** AssignedGroup = Groups.Find(*AssignedGroupName))
+				{
+					(*AssignedGroup)->AddPropertyRow(SubpropertyHandle.ToSharedRef());
+					bHasAssignedGroup = true;
+				}
+			}
+
+			// No group found, so just add it to the root color grading category
+			if (!bHasAssignedGroup)
+			{
+				RootCategory.AddProperty(SubpropertyHandle);
+			}
 		}
 	}
 }
+
+#undef LOCTEXT_NAMESPACE
diff --git a/Engine/Plugins/Experimental/ColorCorrectRegions/Source/ColorCorrectRegionsEditor/Private/ColorCorrectRegionCustomization.h b/Engine/Plugins/Experimental/ColorCorrectRegions/Source/ColorCorrectRegionsEditor/Private/ColorCorrectRegionCustomization.h
--- a/Engine/Plugins/Experimental/ColorCorrectRegions/Source/ColorCorrectRegionsEditor/Private/ColorCorrectRegionCustomization.h
+++ b/Engine/Plugins/Experimental/ColorCorrectRegions/Source/ColorCorrectRegionsEditor/Private/ColorCorrectRegionCustomization.h
@@ -7,6 +7,7 @@
 #include "IDetailCustomization.h"
 #include "PropertyHandle.h"
 
+class IDetailCategoryBuilder;
 
 class FColorCorrectWindowDetails : public IDetailCustomization
 {
@@ -16,4 +17,8 @@ public:
 
 	/** IDetailCustomization interface */
 	virtual void CustomizeDetails(IDetailLayoutBuilder& DetailBuilder) override;
+
+private:
+	/** Recreate the color grading property struct's children as groups or root properties of the color grading category */
+	void MoveColorGradingPropertiesToCategory(TSharedRef<IPropertyHandle> StructHandle, IDetailCategoryBuilder& RootCategory);
 };
diff --git a/Engine/Plugins/Experimental/ColorCorrectRegions/Source/ColorCorrectRegionsEditor/Private/ColorCorrectRegionsEditorModule.cpp b/Engine/Plugins/Experimental/ColorCorrectRegions/Source/ColorCorrectRegionsEditor/Private/ColorCorrectRegionsEditorModule.cpp
--- a/Engine/Plugins/Experimental/ColorCorrectRegions/Source/ColorCorrectRegionsEditor/Private/ColorCorrectRegionsEditorModule.cpp
+++ b/Engine/Plugins/Experimental/ColorCorrectRegions/Source/ColorCorrectRegionsEditor/Private/ColorCorrectRegionsEditorModule.cpp
@@ -11,6 +11,7 @@
 #include "ColorGradingEditorDataModel.h"
 #include "ColorGradingMixerObjectFilterRegistry.h"
 #include "IPlacementModeModule.h"
+#include "PropertyEditorModule.h"
 
 #define LOCTEXT_NAMESPACE "FColorCorrectRegionsModule"
 
@@ -30,10 +31,46 @@ void FColorCorrectRegionsEditorModule::StartupModule()
 
 	FColorGradingMixerObjectFilterRegistry::RegisterObjectClassToFilter(AColorCorrectRegion::StaticClass());
 
+	RegisterSectionMappings();
+
 	ContextMenu = MakeShared<FColorCorrectionActorContextMenu>();
 	ContextMenu->RegisterContextMenuExtender();
 }
 
+void FColorCorrectRegionsEditorModule::RegisterSectionMappings()
+{
+	static const FName PropertyEditor("PropertyEditor");
+	FPropertyEditorModule& PropertyModule = FModuleManager::GetModuleChecked<FPropertyEditorModule>(PropertyEditor);
+
+	const TArray<FName> ClassNames = {
+		AColorCorrectionRegion::StaticClass()->GetFName(),
+		AColorCorrectionWindow::StaticClass()->GetFName()
+	};
+
+	for (const FName& ClassName : ClassNames)
+	{
+		{
+			TSharedRef<FPropertySection> Section = PropertyModule.FindOrCreateSection(ClassName, "Color Grading", LOCTEXT("DetailsSectionColor Grading", "Color Grading"));
+			Section->AddCategory("Color Grading");
+		}
+
+		{
+			TSharedRef<FPropertySection> Section = PropertyModule.FindOrCreateSection(ClassName, "Orientation", LOCTEXT("DetailsSectionOrientation", "Orientation"));
+			Section->AddCategory("Orientation");
+		}
+
+		{
+			TSharedRef<FPropertySection> Section = PropertyModule.FindOrCreateSection(ClassName, "Per Actor CC", LOCTEXT("DetailsSectionPerActorCC", "Per Actor CC"));
+			Section->AddCategory("Per Actor CC");
+		}
+
+		{
+			TSharedRef<FPropertySection> Section = PropertyModule.FindOrCreateSection(ClassName, "Region", LOCTEXT("DetailsSectionRegion", "Region"));
+			Section->AddCategory("Region");
+		}
+	}
+}
+
 void FColorCorrectRegionsEditorModule::OnPlacementModeRefresh(FName CategoryName)
 {
 	static FName VolumeName = FName(TEXT("Volumes"));
diff --git a/Engine/Plugins/Experimental/ColorCorrectRegions/Source/ColorCorrectRegionsEditor/Private/ColorGradingDataModelGenerator_ColorCorrectRegion.cpp b/Engine/Plugins/Experimental/ColorCorrectRegions/Source/ColorCorrectRegionsEditor/Private/ColorGradingDataModelGenerator_ColorCorrectRegion.cpp
--- a/Engine/Plugins/Experimental/ColorCorrectRegions/Source/ColorCorrectRegionsEditor/Private/ColorGradingDataModelGenerator_ColorCorrectRegion.cpp
+++ b/Engine/Plugins/Experimental/ColorCorrectRegions/Source/ColorCorrectRegionsEditor/Private/ColorGradingDataModelGenerator_ColorCorrectRegion.cpp
@@ -4,6 +4,8 @@
 
 #include "ClassIconFinder.h"
 #include "ColorCorrectRegion.h"
+#include "ColorCorrectRegionCustomization.h"
+#include "ColorCorrectWindow.h"
 #include "DetailCategoryBuilder.h"
 #include "DetailLayoutBuilder.h"
 #include "IDetailCustomization.h"
@@ -18,6 +20,11 @@
 
 #define LOCTEXT_NAMESPACE "FColorGradingDataModelGenerator_ColorCorrectRegion"
 
+const TArray<FName> VisibleCategories = {
+	TEXT("Region"),
+	TEXT("Color Grading")
+};
+
 TSharedRef<IColorGradingEditorDataModelGenerator> FColorGradingDataModelGenerator_ColorCorrectRegion::MakeInstance()
 {
 	return MakeShareable(new FColorGradingDataModelGenerator_ColorCorrectRegion());
@@ -28,12 +35,26 @@ class FColorCorrectRegionCustomization : public IDetailCustomization
 public:
 	virtual void CustomizeDetails(IDetailLayoutBuilder& DetailBuilder) override
 	{
+		// Hide any categories whose root isn't in our display list
 		TArray<FName> Categories;
 		DetailBuilder.GetCategoryNames(Categories);
 
 		for (const FName& Category : Categories)
 		{
-			if (Category != TEXT("Color Correction"))
+			const FString RawCategoryName = Category.ToString();
+			FName CategoryRootName;
+			
+			int32 SeparatorIndex;
+			if (RawCategoryName.FindChar('|', SeparatorIndex))
+			{
+				CategoryRootName = FName(RawCategoryName.Left(SeparatorIndex));
+			}
+			else
+			{
+				CategoryRootName = Category;
+			}
+
+			if (!VisibleCategories.Contains(CategoryRootName))
 			{
 				DetailBuilder.HideCategory(Category);
 			}
@@ -42,10 +63,20 @@ public:
 		// TransformCommon is a custom category that doesn't get returned by GetCategoryNames that also needs to be hidden
 		DetailBuilder.HideCategory(TEXT("TransformCommon"));
 
-		IDetailCategoryBuilder& CCCategoryBuilder = DetailBuilder.EditCategory(TEXT("Color Correction"));
+		// Hide CCR-specific properties if CCWs are present in the selection
+		const bool bHasCCWs = DetailBuilder.GetSelectedObjects().ContainsByPredicate([](const TWeakObjectPtr<UObject>& SelectedObject)
+		{
+			return SelectedObject.IsValid() && SelectedObject->IsA<AColorCorrectionWindow>();
+		});
+
+		if (bHasCCWs)
+		{
+			TSharedRef<IPropertyHandle> PriorityProperty = DetailBuilder.GetProperty(GET_MEMBER_NAME_CHECKED(AColorCorrectRegion, Priority));
+			DetailBuilder.HideProperty(PriorityProperty);
 
-		TArray<TSharedRef<IPropertyHandle>> CCPropertyHandles;
-		CCCategoryBuilder.GetDefaultProperties(CCPropertyHandles);
+			TSharedRef<IPropertyHandle> TypeProperty = DetailBuilder.GetProperty(GET_MEMBER_NAME_CHECKED(AColorCorrectRegion, Type));
+			DetailBuilder.HideProperty(TypeProperty);
+		}
 
 		IDetailCategoryBuilder& ColorGradingElementsCategory = DetailBuilder.EditCategory(FName("ColorGradingElements"));
 
@@ -99,7 +130,7 @@ void FColorGradingDataModelGenerator_ColorCorrectRegion::GenerateDataModel(IProp
 		const TSharedRef<IDetailTreeNode> ColorGradingElements = *ColorGradingElementsPtr;
 		FColorGradingEditorDataModel::FColorGradingGroup ColorGradingGroup;
 
-		ColorGradingGroup.DetailsViewCategories.Add(TEXT("Color Correction"));
+		ColorGradingGroup.DetailsViewCategories.Append(VisibleCategories);
 
 		TArray<TSharedRef<IDetailTreeNode>> ColorGradingPropertyNodes;
 		ColorGradingElements->GetChildren(ColorGradingPropertyNodes);
@@ -191,16 +222,4 @@ FColorGradingEditorDataModel::FColorGradingElement FColorGradingDataModelGenerat
 	return ColorGradingElement;
 }
 
-bool FColorGradingDataModelGenerator_ColorCorrectRegion::FilterDetailsViewProperties(const TSharedRef<IDetailTreeNode>& InDetailTreeNode)
-{
-	if (InDetailTreeNode->GetNodeType() == EDetailNodeType::Category)
-	{
-		return InDetailTreeNode->GetNodeName() == TEXT("Color Correction");
-	}
-	else
-	{
-		return true;
-	}
-}
-
 #undef LOCTEXT_NAMESPACE
\ No newline at end of file
diff --git a/Engine/Plugins/Experimental/ColorCorrectRegions/Source/ColorCorrectRegionsEditor/Private/ColorGradingDataModelGenerator_ColorCorrectRegion.h b/Engine/Plugins/Experimental/ColorCorrectRegions/Source/ColorCorrectRegionsEditor/Private/ColorGradingDataModelGenerator_ColorCorrectRegion.h
--- a/Engine/Plugins/Experimental/ColorCorrectRegions/Source/ColorCorrectRegionsEditor/Private/ColorGradingDataModelGenerator_ColorCorrectRegion.h
+++ b/Engine/Plugins/Experimental/ColorCorrectRegions/Source/ColorCorrectRegionsEditor/Private/ColorGradingDataModelGenerator_ColorCorrectRegion.h
@@ -22,6 +22,4 @@ public:
 private:
 	/** Creates a new color grading element structure for the specified detail tree node, which is expected to have child color properties with the ColorGradingMode metadata set */
 	FColorGradingEditorDataModel::FColorGradingElement CreateColorGradingElement(const TSharedRef<IDetailTreeNode>& GroupNode, FText ElementLabel);
-
-	bool FilterDetailsViewProperties(const TSharedRef<IDetailTreeNode>& InDetailTreeNode);
 };
\ No newline at end of file
diff --git a/Engine/Plugins/Experimental/ColorCorrectRegions/Source/ColorCorrectRegionsEditor/Public/ColorCorrectRegionsEditorModule.h b/Engine/Plugins/Experimental/ColorCorrectRegions/Source/ColorCorrectRegionsEditor/Public/ColorCorrectRegionsEditorModule.h
--- a/Engine/Plugins/Experimental/ColorCorrectRegions/Source/ColorCorrectRegionsEditor/Public/ColorCorrectRegionsEditorModule.h
+++ b/Engine/Plugins/Experimental/ColorCorrectRegions/Source/ColorCorrectRegionsEditor/Public/ColorCorrectRegionsEditorModule.h
@@ -16,6 +16,7 @@ public:
 	virtual void ShutdownModule() override;
 	/** IModuleInterface implementation end */
 
+	void RegisterSectionMappings();
 	void OnPlacementModeRefresh(FName CategoryName);
 
 private:
diff --git a/Engine/Plugins/Runtime/nDisplay/Source/DisplayClusterConfigurator/DisplayClusterConfigurator.Build.cs b/Engine/Plugins/Runtime/nDisplay/Source/DisplayClusterConfigurator/DisplayClusterConfigurator.Build.cs
--- a/Engine/Plugins/Runtime/nDisplay/Source/DisplayClusterConfigurator/DisplayClusterConfigurator.Build.cs
+++ b/Engine/Plugins/Runtime/nDisplay/Source/DisplayClusterConfigurator/DisplayClusterConfigurator.Build.cs
@@ -22,6 +22,7 @@ public class DisplayClusterConfigurator : ModuleRules
 				"AssetTools",
 				"CinematicCamera",
 				"ClassViewer",
+				"ColorGradingEditor",
 				"Core",
 				"CoreUObject",
 				"DesktopPlatform",
diff --git a/Engine/Plugins/Runtime/nDisplay/Source/DisplayClusterConfigurator/Private/Views/Details/Components/DisplayClusterICVFXCameraComponentDetailsCustomization.cpp b/Engine/Plugins/Runtime/nDisplay/Source/DisplayClusterConfigurator/Private/Views/Details/Components/DisplayClusterICVFXCameraComponentDetailsCustomization.cpp
--- a/Engine/Plugins/Runtime/nDisplay/Source/DisplayClusterConfigurator/Private/Views/Details/Components/DisplayClusterICVFXCameraComponentDetailsCustomization.cpp
+++ b/Engine/Plugins/Runtime/nDisplay/Source/DisplayClusterConfigurator/Private/Views/Details/Components/DisplayClusterICVFXCameraComponentDetailsCustomization.cpp
@@ -10,6 +10,7 @@
 #include "DisplayClusterRootActor.h"
 #include "UObject/SoftObjectPtr.h"
 
+#include "ColorGradingEditorUtil.h"
 #include "DetailCategoryBuilder.h"
 #include "DetailLayoutBuilder.h"
 #include "IDetailGroup.h"
@@ -98,7 +99,14 @@ void FDisplayClusterICVFXCameraComponentDetailsCustomization::CustomizeDetails(I
 	InLayoutBuilder.HideCategory(TEXT("Sockets"));
 
 	// Rename "Inner Frustum Color Grading" to "Color Grading" for brevity, as the category itself needs to remain distinct from the camera's "Color Grading" category.
-	InLayoutBuilder.EditCategory(DisplayClusterConfigurationStrings::categories::CameraColorGradingCategory, LOCTEXT("ICVFXColorGradingCategoryLabel", "Color Grading"));
+	IDetailCategoryBuilder& ColorGradingCategory = InLayoutBuilder.EditCategory(DisplayClusterConfigurationStrings::categories::CameraColorGradingCategory, LOCTEXT("ICVFXColorGradingCategoryLabel", "Color Grading"));
+
+	// Add the Color Grading button at the top of the new category
+	ColorGradingCategory.AddCustomRow(NSLOCTEXT("ColorCorrectWindowDetails", "OpenColorGrading", "Open Color Grading"))
+		.RowTag("OpenColorGrading")
+		[
+			ColorGradingEditorUtil::MakeColorGradingLaunchButton()
+		];
 
 	IDetailCategoryBuilder& CameraCategory = InLayoutBuilder.EditCategory(DisplayClusterConfigurationStrings::categories::ICVFXCameraCategory, LOCTEXT("ICVFXCameraCategoryLabel", "Camera"));
 
diff --git a/Engine/Plugins/Runtime/nDisplay/Source/DisplayClusterConfigurator/Private/Views/Details/DisplayClusterRootActorDetailsCustomization.cpp b/Engine/Plugins/Runtime/nDisplay/Source/DisplayClusterConfigurator/Private/Views/Details/DisplayClusterRootActorDetailsCustomization.cpp
--- a/Engine/Plugins/Runtime/nDisplay/Source/DisplayClusterConfigurator/Private/Views/Details/DisplayClusterRootActorDetailsCustomization.cpp
+++ b/Engine/Plugins/Runtime/nDisplay/Source/DisplayClusterConfigurator/Private/Views/Details/DisplayClusterRootActorDetailsCustomization.cpp
@@ -16,6 +16,7 @@
 #include "IDetailGroup.h"
 #include "PropertyHandle.h"
 
+#include "ColorGradingEditorUtil.h"
 #include "EditorSupportDelegates.h"
 #include "SSearchableComboBox.h"
 #include "Widgets/Text/STextBlock.h"
@@ -78,6 +79,16 @@ void FDisplayClusterRootActorDetailsCustomization::CustomizeDetails(IDetailLayou
 {
 	FDisplayClusterConfiguratorBaseDetailCustomization::CustomizeDetails(InLayoutBuilder);
 
+	// Add the Color Grading button at the top of the relevant category
+	{
+		IDetailCategoryBuilder& ColorGradingCategory = InLayoutBuilder.EditCategory(DisplayClusterConfigurationStrings::categories::ColorGradingCategory, LOCTEXT("ColorGradingDetails", "Color Grading"));
+		ColorGradingCategory.AddCustomRow(NSLOCTEXT("ColorCorrectWindowDetails", "OpenColorGrading", "Open Color Grading"))
+			.RowTag("OpenColorGrading")
+			[
+				ColorGradingEditorUtil::MakeColorGradingLaunchButton()
+			];
+	}
+
 	const TArray<TWeakObjectPtr<UObject>>& SelectedObjects = InLayoutBuilder.GetSelectedObjects();
 	bMultipleObjectsSelected = SelectedObjects.Num() > 1;
 
