Color Correction Regions: Removed unnecessary refresh of stencil values that now causes issues on snapshot load and multi-user transactions.

#rnx

diff --git a/Engine/Plugins/Experimental/ColorCorrectRegions/Source/ColorCorrectRegions/Private/ColorCorrectRegionSubsystem.cpp b/Engine/Plugins/Experimental/ColorCorrectRegions/Source/ColorCorrectRegions/Private/ColorCorrectRegionSubsystem.cpp
--- a/Engine/Plugins/Experimental/ColorCorrectRegions/Source/ColorCorrectRegions/Private/ColorCorrectRegionSubsystem.cpp
+++ b/Engine/Plugins/Experimental/ColorCorrectRegions/Source/ColorCorrectRegions/Private/ColorCorrectRegionSubsystem.cpp
@@ -367,23 +367,6 @@ void UColorCorrectRegionsSubsystem::CheckAssignedActorsValidity(AColorCorrectReg
 	FColorCorrectRegionsStencilManager::CheckAssignedActorsValidity(Region);
 }
 
-void UColorCorrectRegionsSubsystem::RefreshStenciIdAssignmentForAllCCR()
-{
-	for (TActorIterator<AColorCorrectRegion> It(GetWorld()); It; ++It)
-	{
-		AColorCorrectRegion* AsRegion = *It;
-		if (IsRegionValid(AsRegion, GetWorld()))
-		{
-			// Uncoment this if you want the invalid actors to be removed automatically.
-			// However after removal, undo/redo will not re-assign this actor back to CCR.
-			//FColorCorrectRegionsStencilManager::ClearInvalidActorsForSelectedRegion(AsRegion);
-			AsRegion->PerAffectedActorStencilData.Empty();
-			FColorCorrectRegionsStencilManager::AssignStencilIdsToAllActorsForCCR(GetWorld(), AsRegion, true, true);
-			FColorCorrectRegionsStencilManager::RemoveStencilNumberForSelectedRegion(GetWorld(), AsRegion);
-		}
-	}
-}
-
 void UColorCorrectRegionsSubsystem::RefreshRegions()
 {
 	RegionsPriorityBased.Reset();
@@ -403,8 +386,6 @@ void UColorCorrectRegionsSubsystem::RefreshRegions()
 			}
 		}
 	}
-
-	RefreshStenciIdAssignmentForAllCCR();
 }
 
 #undef LOCTEXT_NAMESPACE
\ No newline at end of file
diff --git a/Engine/Plugins/Experimental/ColorCorrectRegions/Source/ColorCorrectRegions/Public/ColorCorrectRegionsSubsystem.h b/Engine/Plugins/Experimental/ColorCorrectRegions/Source/ColorCorrectRegions/Public/ColorCorrectRegionsSubsystem.h
--- a/Engine/Plugins/Experimental/ColorCorrectRegions/Source/ColorCorrectRegions/Public/ColorCorrectRegionsSubsystem.h
+++ b/Engine/Plugins/Experimental/ColorCorrectRegions/Source/ColorCorrectRegions/Public/ColorCorrectRegionsSubsystem.h
@@ -69,8 +69,9 @@ public:
 	void CheckAssignedActorsValidity(AColorCorrectRegion* Region);
 
 	/** Resets all stencils and re-assigns for each CCR in the scene. */
+	UE_DEPRECATED(5.5, "Refreshing stencil IDs per tick no longer necessary as we check validity.")
 	UFUNCTION(BlueprintCallable, meta = (Category = "Color Correct Regions"))
-	void RefreshStenciIdAssignmentForAllCCR();
+	void RefreshStenciIdAssignmentForAllCCR() {};
 
 public:
 #if WITH_EDITOR
