[ColorCorrectRegions] Add context menu items for adding actors to a CCR's per-actor CC, and for removing them from the actor's list item

diff --git a/Engine/Plugins/Editor/ColorGrading/Source/ColorGradingEditor/ColorGradingEditor.build.cs b/Engine/Plugins/Editor/ColorGrading/Source/ColorGradingEditor/ColorGradingEditor.build.cs
--- a/Engine/Plugins/Editor/ColorGrading/Source/ColorGradingEditor/ColorGradingEditor.build.cs
+++ b/Engine/Plugins/Editor/ColorGrading/Source/ColorGradingEditor/ColorGradingEditor.build.cs
@@ -32,6 +32,7 @@ public class ColorGradingEditor : ModuleRules
 				"PropertyPath",
 				"Slate",
 				"SlateCore",
+				"ToolMenus",
 				"ToolWidgets",
 				"UnrealEd",
 				"WorkspaceMenuStructure",
diff --git a/Engine/Plugins/Editor/ColorGrading/Source/ColorGradingEditor/Private/ColorGradingMixerObjectFilter.cpp b/Engine/Plugins/Editor/ColorGrading/Source/ColorGradingEditor/Private/ColorGradingMixerObjectFilter.cpp
--- a/Engine/Plugins/Editor/ColorGrading/Source/ColorGradingEditor/Private/ColorGradingMixerObjectFilter.cpp
+++ b/Engine/Plugins/Editor/ColorGrading/Source/ColorGradingEditor/Private/ColorGradingMixerObjectFilter.cpp
@@ -4,8 +4,10 @@
 
 #include "ActorTreeItem.h"
 #include "ComponentTreeItem.h"
+#include "ColorGradingMixerContextObject.h"
 #include "ColorGradingMixerObjectFilterRegistry.h"
 #include "ISceneOutliner.h"
+#include "ToolMenuContext.h"
 
 TSet<UClass*> UColorGradingMixerObjectFilter::GetObjectClassesToFilter() const
 {
@@ -87,6 +89,26 @@ void UColorGradingMixerObjectFilter::OnDrop(ISceneOutlinerTreeItem& DropTarget,
 	}
 }
 
+TSet<FName> UColorGradingMixerObjectFilter::GetPropertiesThatRequireListRefresh() const
+{
+	TSet<FName> PropertyNames;
+
+	for (UClass* Class : FColorGradingMixerObjectFilterRegistry::GetObjectClassesToFilter())
+	{
+		if (const IColorGradingMixerObjectHierarchyConfig* Config = FColorGradingMixerObjectFilterRegistry::GetClassHierarchyConfig(Class))
+		{
+			PropertyNames.Append(Config->GetPropertiesThatRequireListRefresh());
+		}
+	}
+
+	return MoveTemp(PropertyNames);
+}
+
+void UColorGradingMixerObjectFilter::OnContextMenuContextCreated(FToolMenuContext& Context) const
+{
+	Context.AddObject(NewObject<UColorGradingMixerContextObject>());
+}
+
 UObject* UColorGradingMixerObjectFilter::GetObjectForTreeItem(const ISceneOutlinerTreeItem& TreeItem) const
 {
 	if (const FActorTreeItem* ActorTreeItem = TreeItem.CastTo<FActorTreeItem>())
diff --git a/Engine/Plugins/Editor/ColorGrading/Source/ColorGradingEditor/Private/ColorGradingMixerObjectFilter.h b/Engine/Plugins/Editor/ColorGrading/Source/ColorGradingEditor/Private/ColorGradingMixerObjectFilter.h
--- a/Engine/Plugins/Editor/ColorGrading/Source/ColorGradingEditor/Private/ColorGradingMixerObjectFilter.h
+++ b/Engine/Plugins/Editor/ColorGrading/Source/ColorGradingEditor/Private/ColorGradingMixerObjectFilter.h
@@ -19,6 +19,8 @@ public:
 	virtual bool HasCustomDropHandling(const ISceneOutlinerTreeItem& DropTarget) const override;
 	virtual FSceneOutlinerDragValidationInfo ValidateDrop(const ISceneOutlinerTreeItem& DropTarget, const FSceneOutlinerDragDropPayload& Payload) const override;
 	virtual void OnDrop(ISceneOutlinerTreeItem& DropTarget, const FSceneOutlinerDragDropPayload& Payload, const FSceneOutlinerDragValidationInfo& ValidationInfo) const override;
+	virtual TSet<FName> GetPropertiesThatRequireListRefresh() const override;
+	virtual void OnContextMenuContextCreated(FToolMenuContext& Context) const override;
 
 	virtual bool GetShowTransientObjects() const override
 	{
diff --git a/Engine/Plugins/Editor/ColorGrading/Source/ColorGradingEditor/Public/ColorGradingMixerContextObject.h b/Engine/Plugins/Editor/ColorGrading/Source/ColorGradingEditor/Public/ColorGradingMixerContextObject.h
new file mode 100644
--- /dev/null
+++ b/Engine/Plugins/Editor/ColorGrading/Source/ColorGradingEditor/Public/ColorGradingMixerContextObject.h
@@ -0,0 +1,14 @@
+// Copyright Epic Games, Inc. All Rights Reserved.
+
+#pragma once
+
+#include "ColorGradingMixerContextObject.generated.h"	
+
+class UToolMenu;
+
+/** Empty class which can be added to an FToolMenuContext to indicate that the menu originates from the Color Grading panel. */
+UCLASS()
+class COLORGRADINGEDITOR_API UColorGradingMixerContextObject : public UObject
+{
+	GENERATED_BODY()
+};
diff --git a/Engine/Plugins/Editor/ColorGrading/Source/ColorGradingEditor/Public/ColorGradingMixerObjectFilterRegistry.h b/Engine/Plugins/Editor/ColorGrading/Source/ColorGradingEditor/Public/ColorGradingMixerObjectFilterRegistry.h
--- a/Engine/Plugins/Editor/ColorGrading/Source/ColorGradingEditor/Public/ColorGradingMixerObjectFilterRegistry.h
+++ b/Engine/Plugins/Editor/ColorGrading/Source/ColorGradingEditor/Public/ColorGradingMixerObjectFilterRegistry.h
@@ -35,7 +35,13 @@ public:
 	virtual FSceneOutlinerDragValidationInfo ValidateDrop(UObject* DropTarget, const FSceneOutlinerDragDropPayload& Payload) const { return FSceneOutlinerDragValidationInfo::Invalid(); }
 
 	/** Called when a payload is dropped onto a target. If this returns true, the outliner will be refreshed after the operation. */
-	virtual bool OnDrop(UObject* DropTarget, const FSceneOutlinerDragDropPayload& Payload, const FSceneOutlinerDragValidationInfo& ValidationInfo) const { return false;  }
+	virtual bool OnDrop(UObject* DropTarget, const FSceneOutlinerDragDropPayload& Payload, const FSceneOutlinerDragValidationInfo& ValidationInfo) const { return false; }
+
+	/**
+	 * If a property is changed that has a name found in this set, the color grading mixer hierarchy will be refreshed.
+	 * Add a property name to this list if you expect the list to change in some way after changing that property.
+	 */
+	virtual TSet<FName> GetPropertiesThatRequireListRefresh() const { return {}; }
 };
 
 DECLARE_DELEGATE_RetVal(TSharedRef<IColorGradingMixerObjectHierarchyConfig>, FGetObjectHierarchyConfig);
diff --git a/Engine/Plugins/Editor/ObjectMixer/ObjectMixer/Source/ObjectMixer/Private/Views/List/Modes/ObjectMixerOutlinerMode.cpp b/Engine/Plugins/Editor/ObjectMixer/ObjectMixer/Source/ObjectMixer/Private/Views/List/Modes/ObjectMixerOutlinerMode.cpp
--- a/Engine/Plugins/Editor/ObjectMixer/ObjectMixer/Source/ObjectMixer/Private/Views/List/Modes/ObjectMixerOutlinerMode.cpp
+++ b/Engine/Plugins/Editor/ObjectMixer/ObjectMixer/Source/ObjectMixer/Private/Views/List/Modes/ObjectMixerOutlinerMode.cpp
@@ -1876,6 +1876,11 @@ TSharedPtr<SWidget> FObjectMixerOutlinerMode::BuildContextMenu()
 	UObjectMixerEditorListMenuContext::RegisterObjectMixerDynamicCollectionsContextMenuExtension("LevelEditor.ComponentContextMenu");
 	UObjectMixerEditorListMenuContext::RegisterObjectMixerDynamicCollectionsContextMenuExtension("LevelEditor.ElementContextMenu");
 
+	if (const UObjectMixerObjectFilter* Filter = GetListModelPtr().Pin()->GetMainObjectFilterInstance())
+	{
+		Filter->OnContextMenuContextCreated(Context);
+	}
+
 	const FName MenuName = "LevelEditor.LevelEditorSceneOutliner.ContextMenu";
 	
 	const FLevelEditorModule& LevelEditorModule = FModuleManager::LoadModuleChecked<FLevelEditorModule>("LevelEditor");
diff --git a/Engine/Plugins/Editor/ObjectMixer/ObjectMixer/Source/ObjectMixer/Public/ObjectFilter/ObjectMixerEditorObjectFilter.h b/Engine/Plugins/Editor/ObjectMixer/ObjectMixer/Source/ObjectMixer/Public/ObjectFilter/ObjectMixerEditorObjectFilter.h
--- a/Engine/Plugins/Editor/ObjectMixer/ObjectMixer/Source/ObjectMixer/Public/ObjectFilter/ObjectMixerEditorObjectFilter.h
+++ b/Engine/Plugins/Editor/ObjectMixer/ObjectMixer/Source/ObjectMixer/Public/ObjectFilter/ObjectMixerEditorObjectFilter.h
@@ -8,6 +8,8 @@
 
 #include "ObjectMixerEditorObjectFilter.generated.h"
 
+struct FToolMenuContext;
+
 UENUM(BlueprintType)
 enum class EObjectMixerInheritanceInclusionOptions : uint8
 {
@@ -160,6 +162,8 @@ public:
 	/** Called when a payload is dropped onto a target */
 	virtual void OnDrop(ISceneOutlinerTreeItem& DropTarget, const FSceneOutlinerDragDropPayload& Payload, const FSceneOutlinerDragValidationInfo& ValidationInfo) const;
 
+	/** Called when a context menu is being created from in the hierarchy. Allows subclasses to inject additonal context data. */
+	virtual void OnContextMenuContextCreated(FToolMenuContext& Context) const {}
 
 	static TSet<UClass*> GetParentAndChildClassesFromSpecifiedClasses(
 		const TSet<UClass*>& InSpecifiedClasses, EObjectMixerInheritanceInclusionOptions Options);
diff --git a/Engine/Plugins/Experimental/ColorCorrectRegions/ColorCorrectRegions.uplugin b/Engine/Plugins/Experimental/ColorCorrectRegions/ColorCorrectRegions.uplugin
--- a/Engine/Plugins/Experimental/ColorCorrectRegions/ColorCorrectRegions.uplugin
+++ b/Engine/Plugins/Experimental/ColorCorrectRegions/ColorCorrectRegions.uplugin
@@ -34,6 +34,10 @@
 		{
 			"Name": "ColorGrading",
 			"Enabled": true
+		},
+		{
+			"Name": "ObjectMixer",
+			"Enabled": true
 		}
 	]
 }
\ No newline at end of file
diff --git a/Engine/Plugins/Experimental/ColorCorrectRegions/Source/ColorCorrectRegionsEditor/ColorCorrectRegionsEditor.Build.cs b/Engine/Plugins/Experimental/ColorCorrectRegions/Source/ColorCorrectRegionsEditor/ColorCorrectRegionsEditor.Build.cs
--- a/Engine/Plugins/Experimental/ColorCorrectRegions/Source/ColorCorrectRegionsEditor/ColorCorrectRegionsEditor.Build.cs
+++ b/Engine/Plugins/Experimental/ColorCorrectRegions/Source/ColorCorrectRegionsEditor/ColorCorrectRegionsEditor.Build.cs
@@ -13,6 +13,7 @@ public class ColorCorrectRegionsEditor : ModuleRules
 				"Core",
 				"CoreUObject",
 				"Engine",
+				"LevelEditor",
 				"Projects",
 				"UnrealEd",
 				"Slate",
@@ -21,6 +22,9 @@ public class ColorCorrectRegionsEditor : ModuleRules
 				"ColorGradingEditor",
 				"PlacementMode",
 				"SceneOutliner",
+				"ObjectMixerEditor",
+				"ToolMenus",
+				"TypedElementRuntime"
 			}
 		);
 	}
diff --git a/Engine/Plugins/Experimental/ColorCorrectRegions/Source/ColorCorrectRegionsEditor/Private/ColorCorrectionActorContextMenu.cpp b/Engine/Plugins/Experimental/ColorCorrectRegions/Source/ColorCorrectRegionsEditor/Private/ColorCorrectionActorContextMenu.cpp
--- a/Engine/Plugins/Experimental/ColorCorrectRegions/Source/ColorCorrectRegionsEditor/Private/ColorCorrectionActorContextMenu.cpp
+++ b/Engine/Plugins/Experimental/ColorCorrectRegions/Source/ColorCorrectRegionsEditor/Private/ColorCorrectionActorContextMenu.cpp
@@ -5,9 +5,10 @@
 #include "ActorTreeItem.h"
 #include "ColorCorrectRegion.h"
 #include "ColorCorrectWindow.h"
+#include "ColorGradingMixerContextObject.h"
 #include "Kismet/GameplayStatics.h"
 #include "LevelEditor.h"
-#include "LevelEditor.h"
+#include "LevelEditorMenuContext.h"
 #include "Editor/SceneOutliner/Public/ActorMode.h"
 #include "Engine/Selection.h"
 #include "Framework/Application/SlateApplication.h"
@@ -18,6 +19,14 @@
 #include "Editor/UnrealEdEngine.h"
 #include "UnrealEdGlobals.h"
 #include "Modules/ModuleManager.h"
+#include "ScopedTransaction.h"
+#include "ToolMenu.h"
+#include "ToolMenuContext.h"
+#include "ToolMenuDelegates.h"
+#include "ToolMenuEntry.h"
+#include "ToolMenus.h"
+#include "ToolMenuSection.h"
+#include "Views/Widgets/ObjectMixerEditorListMenuContext.h"
 
 #define LOCTEXT_NAMESPACE "FColorCorrectRegionsContextMenu"
 
@@ -46,27 +55,52 @@ namespace
 		FOnSceneOutlinerItemPicked OnCCActorPicked;
 	};
 
+	// Map from actor to parent CCRs from which to remove the actor
+	using TPerActorCCRemovalMap = TMap<AActor*, TArray<AColorCorrectRegion*>>;
 
-	/** Add Selected actors to the chosen CC Actor's Per actor CC. */
-	static void AddActorsToPerActorCC(AColorCorrectRegion* CCActorPtr, TSharedPtr<TArray<AActor*>> InSelectedActors)
+	/** Add Selected actors to the chosen CC Actors' Per actor CC. */
+	static void AddActorsToPerActorCC(TArray<AActor*> InCCActors, TArray<AActor*> InAffectedActors)
 	{
-		// Hide all context menus.
-		FSlateApplication::Get().DismissAllMenus();
+		FScopedTransaction Transaction(LOCTEXT("AddToPerActorCCTransaction", "Add actors to Per-Actor CC"));
 
-		CCActorPtr->bEnablePerActorCC = true;
-		for (AActor* SelectedActor : *InSelectedActors)
+		for (AActor* CCActor : InCCActors)
 		{
-			TSoftObjectPtr<AActor> SelectedActorPtr(SelectedActor);
+			if (AColorCorrectRegion* CCActorPtr = Cast<AColorCorrectRegion>(CCActor))
+			{
+				CCActorPtr->Modify();
+
+				CCActorPtr->bEnablePerActorCC = true;
+				for (AActor* SelectedActor : InAffectedActors)
+				{
+					TSoftObjectPtr<AActor> SelectedActorPtr(SelectedActor);
+					CCActorPtr->AffectedActors.Add(SelectedActorPtr);
+				}
+			}
+		}
+	}
 
-			// Forcing property change event.
-			const FName AffectedActorsPropertyName = GET_MEMBER_NAME_CHECKED(AColorCorrectRegion, AffectedActors);
-			FProperty* AffectedActorsProperty = FindFieldChecked<FProperty>(AColorCorrectRegion::StaticClass(), AffectedActorsPropertyName);
-			CCActorPtr->PreEditChange(AffectedActorsProperty);
+	/** Remove selected actors from the CC Actors' Per actor CC. */
+	static void RemoveActorsFromPerActorCC(TSharedPtr<TPerActorCCRemovalMap> RemovalMap)
+	{
+		FScopedTransaction Transaction(LOCTEXT("AddToPerActorCCTransaction", "Remove actors from Per-Actor CC"));
 
-			CCActorPtr->AffectedActors.Add(SelectedActorPtr);
-			FPropertyChangedEvent PropertyEvent(AffectedActorsProperty);
-			PropertyEvent.ChangeType = EPropertyChangeType::ArrayAdd;
-			CCActorPtr->PostEditChangeProperty(PropertyEvent);
+		for (TPair<AActor*, TArray<AColorCorrectRegion*>> Pair : *RemovalMap)
+		{
+			AActor* ActorToRemove = Pair.Key;
+			TArray<AColorCorrectRegion*>& CCActors = Pair.Value;
+
+			for (AColorCorrectRegion* CCActor : CCActors)
+			{
+				CCActor->Modify();
+
+				CCActor->AffectedActors.Remove(ActorToRemove);
+
+				// If no affected actors are left, disable per-actor CC so it won't become invisible
+				if (CCActor->AffectedActors.IsEmpty())
+				{
+					CCActor->bEnablePerActorCC = false;
+				}
+			}
 		}
 	}
 
@@ -119,23 +153,38 @@ namespace
 		}
 
 		CCActorPtr->SetActorTransform(Transform);
-		AddActorsToPerActorCC(CCActorPtr, SelectedActors);
+		AddActorsToPerActorCC({ CCActorPtr }, *SelectedActors);
+
+		// Hide all context menus.
+		FSlateApplication::Get().DismissAllMenus();
 
 		// Shift selection to the newly created CC Actor.
 		GEditor->SelectNone(false/*bNoteSelectionChange*/, true/*bDeselectBSPSurfs*/);
 		GEditor->SelectActor(CCActorPtr, true/*bInSelected*/, true/*bNotify*/);
 	};
 
-	/** Called by outliner when a tree item is selected by the user. */
-	static void OnOutlinerTreeItemSelected(TSharedRef<ISceneOutlinerTreeItem> NewParent, TSharedPtr<TArray<AActor*>> SelectedActors)
+	/** Called by outliner when a new affected actor is selected to be added to one or more CC actors */
+	static void OnAddPerActorCCFromCCRTreeItemSelected(TSharedRef<ISceneOutlinerTreeItem> NewAffectedActor, TSharedPtr<TArray<AActor*>> CCActors)
 	{
-		if (FActorTreeItem* ActorItem = NewParent->CastTo<FActorTreeItem>())
+		if (FActorTreeItem* ActorItem = NewAffectedActor->CastTo<FActorTreeItem>())
 		{
-			if (AColorCorrectRegion* CCR = Cast<AColorCorrectRegion>(ActorItem->Actor))
-			{
-				AddActorsToPerActorCC(CCR, SelectedActors);
-			}
+			AddActorsToPerActorCC(*CCActors, { ActorItem->Actor.Get() });
 		}
+
+		// Hide all context menus.
+		FSlateApplication::Get().DismissAllMenus();
+	};
+
+	/** Called by outliner when a CC actor is selected to have the affected actors added to its Per-Actor CC array */
+	static void OnAddPerActorCCFromActorTreeItemSelected(TSharedRef<ISceneOutlinerTreeItem> CCActor, TSharedPtr<TArray<AActor*>> AffectedActors)
+	{
+		if (FActorTreeItem* ActorItem = CCActor->CastTo<FActorTreeItem>())
+		{
+			AddActorsToPerActorCC({ ActorItem->Actor.Get()}, *AffectedActors);
+		}
+
+		// Hide all context menus.
+		FSlateApplication::Get().DismissAllMenus();
 	};
 
 	/** A helper function used to determine if selected actor is either CCR or CCW. */
@@ -144,7 +193,18 @@ namespace
 		return (Cast<AColorCorrectRegion>(InActor) != nullptr);
 	}
 
-	/** Transfers Editor into a picker state when users select a pipet button. */
+	/** Returns true if this actor can be added to any of the selected actors' per-actor CC lists. */
+	static bool CanActorBeAddedToPerActorCC(const AActor* const InActor)
+	{
+		if (IsActorCCR(InActor))
+		{
+			return false;
+		}
+
+		return InActor->GetComponentByClass<UPrimitiveComponent>() != nullptr;
+	}
+
+	/** Transfers Editor into a picker state for selecting a color correction actor. */
 	static FReply PickCCActorMode(TSharedPtr<TArray<AActor*>> InSelectedActors)
 	{
 		FSlateApplication::Get().DismissAllMenus();
@@ -153,128 +213,295 @@ namespace
 		ActorPickerModeModule.BeginActorPickingMode(
 			FOnGetAllowedClasses(),
 			FOnShouldFilterActor::CreateStatic(IsActorCCR),
-			FOnActorSelected::CreateLambda([](AActor* InActor, TSharedPtr<TArray<AActor*>> SelectedActors)
+			FOnActorSelected::CreateLambda([InSelectedActors](AActor* InCCActor)
 				{
-					AddActorsToPerActorCC(Cast<AColorCorrectRegion>(InActor), SelectedActors);
-				}, 
-				InSelectedActors)
+					AddActorsToPerActorCC({ InCCActor }, *InSelectedActors);
+				})
 		);
 		return FReply::Handled();
 	}
 
-	/** Creates and populates menu allowing users to add selected actors to Per Actor CC. */
-	static void AddAttachToPerActorCCRMenu(class FMenuBuilder& MenuBuilder, TSharedPtr<TArray<AActor*>> InSelectedActors)
+	/** Transfers Editor into a picker state for selecting an actor to add to a color correction actor's per-actor CC list. */
+	static FReply PickActorToAddToCCMode(TSharedPtr<TArray<AActor*>> InCCActors)
 	{
-		MenuBuilder.BeginSection("ColorCorrectionRegionsSection", LOCTEXT("ColorCorrectionRegions", "Color Correction Regions"));
+		FSlateApplication::Get().DismissAllMenus();
+		FActorPickerModeModule& ActorPickerModeModule = FModuleManager::Get().GetModuleChecked<FActorPickerModeModule>("ActorPickerMode");
 
-		auto CCActorAddToPerActorSubMenu = [](class FMenuBuilder& MenuBuilder, TSharedPtr<TArray<AActor*>> InSelectedActors)
-		{
-			FUIAction Action_CreateNewCCW(FExecuteAction::CreateStatic(CreateNewCCActor, Window, InSelectedActors));
-			FUIAction Action_CreateNewCCR(FExecuteAction::CreateStatic(CreateNewCCActor, Region, InSelectedActors));
-
-			MenuBuilder.AddMenuEntry(
-				LOCTEXT("MenuCreateAttachCCW", "Add to New Color Correction Window"),
-				LOCTEXT("MenuCreateAttachCCW_Tooltip", "Creates new Color Correction Window (CCW) and adds valid selected actors to Per-Actor CC of the newly created CCW."),
-				FSlateIcon(),
-				Action_CreateNewCCW,
-				NAME_None,
-				EUserInterfaceActionType::Button);
-
-			MenuBuilder.AddMenuEntry(
-				LOCTEXT("MenuCreateAttachCCR", "Add to New Color Correction Region"),
-				LOCTEXT("MenuCreateAttachCCR_Tooltip", "Creates new Color Correction Region (CCR) and adds valid selected actors to Per-Actor CC of the newly created CCR."),
-				FSlateIcon(),
-				Action_CreateNewCCR,
-				NAME_None,
-				EUserInterfaceActionType::Button);
-
-			MenuBuilder.AddMenuSeparator();
-
-			FCreateSceneOutlinerMode ModeFactory = FCreateSceneOutlinerMode::CreateLambda([](SSceneOutliner* Outliner, TSharedPtr<TArray<AActor*>> SelectedActors)
+		ActorPickerModeModule.BeginActorPickingMode(
+			FOnGetAllowedClasses(),
+			FOnShouldFilterActor::CreateStatic(CanActorBeAddedToPerActorCC),
+			FOnActorSelected::CreateLambda([InCCActors](AActor* InSelectedActor)
 				{
-					return new FCCActorPickingMode(Outliner, FOnSceneOutlinerItemPicked::CreateStatic(OnOutlinerTreeItemSelected, SelectedActors));
-				}, InSelectedActors);
-
-			FSceneOutlinerInitializationOptions CCSceneOutlinerInitOptions;
-			CCSceneOutlinerInitOptions.bShowHeaderRow = false;
-			CCSceneOutlinerInitOptions.bFocusSearchBoxWhenOpened = true;
-			CCSceneOutlinerInitOptions.ModeFactory = ModeFactory;
-			CCSceneOutlinerInitOptions.Filters->AddFilterPredicate<FActorTreeItem>(FActorTreeItem::FFilterPredicate::CreateStatic(IsActorCCR));
-
-			TSharedRef<SWidget> CCSceneOutliner =
-				SNew(SHorizontalBox)
-				+ SHorizontalBox::Slot()
+					AddActorsToPerActorCC(*InCCActors, { InSelectedActor });
+				})
+		);
+		return FReply::Handled();
+	}
+
+	/** Build the menu that adds an actor to CCRs when right-clicking on the CCRs */
+	static void CreateAddPerActorCCFromCCRMenu(UToolMenu* Menu, TSharedPtr<TArray<AActor*>> InSelectedActors)
+	{
+		FCreateSceneOutlinerMode ModeFactory = FCreateSceneOutlinerMode::CreateLambda([](SSceneOutliner* Outliner, TSharedPtr<TArray<AActor*>> SelectedActors)
+		{
+			return new FCCActorPickingMode(Outliner, FOnSceneOutlinerItemPicked::CreateStatic(OnAddPerActorCCFromCCRTreeItemSelected, SelectedActors));
+		}, InSelectedActors);
+
+		FSceneOutlinerInitializationOptions TargetSceneOutlinerInitOptions;
+		TargetSceneOutlinerInitOptions.bShowHeaderRow = false;
+		TargetSceneOutlinerInitOptions.bFocusSearchBoxWhenOpened = true;
+		TargetSceneOutlinerInitOptions.ModeFactory = ModeFactory;
+		TargetSceneOutlinerInitOptions.Filters->AddFilterPredicate<FActorTreeItem>(FActorTreeItem::FFilterPredicate::CreateStatic(CanActorBeAddedToPerActorCC));
+
+		TSharedRef<SWidget> CCSceneOutliner =
+			SNew(SHorizontalBox)
+			+ SHorizontalBox::Slot()
+			[
+				SNew(SSceneOutliner, TargetSceneOutlinerInitOptions)
+				.IsEnabled(FSlateApplication::Get().GetNormalExecutionAttribute())
+			]
+
+			+ SHorizontalBox::Slot()
+			.VAlign(VAlign_Top)
+			.AutoWidth()
+			[
+				SNew(SVerticalBox)
+
+				+ SVerticalBox::Slot()
+				.AutoHeight()
+				.Padding(0.0f, 9.0f, 0.0f, 0.0f)
 				[
-					SNew(SSceneOutliner, CCSceneOutlinerInitOptions)
-					.IsEnabled(FSlateApplication::Get().GetNormalExecutionAttribute())
+					SNew(SButton)
+					.ToolTipText(LOCTEXT("PickActorButtonLabel", "Pick an Actor"))
+					.ButtonStyle(FAppStyle::Get(), "HoverHintOnly")
+					.OnClicked(FOnClicked::CreateStatic(PickActorToAddToCCMode, InSelectedActors))
+					.ForegroundColor(FSlateColor::UseForeground())
+					.IsFocusable(false)
+					[
+						SNew(SImage)
+						.Image(FAppStyle::GetBrush("Icons.EyeDropper"))
+						.ColorAndOpacity(FSlateColor::UseForeground())
+					]
 				]
+			];
+		
+		FToolMenuSection& ExistingSection = Menu->AddSection("Existing", LOCTEXT("ExistingCCActorSection", "Existing"));
+		ExistingSection.AddEntry(
+			FToolMenuEntry::InitWidget("Picker", CCSceneOutliner, FText::GetEmpty(), /*bNoIndent=*/ true)
+		);
+	};
 
-				+SHorizontalBox::Slot()
-				.VAlign(VAlign_Top)
-				.AutoWidth()
-				[
-					SNew(SVerticalBox)
+	/** Build the menu that adds actors to a CCR's per-actor CC when right-clicking on an actor */
+	static void CreateAddPerActorCCFromActorMenu(UToolMenu* Menu, TSharedPtr<TArray<AActor*>> InSelectedActors)
+	{
+		FToolMenuSection& TopSection = Menu->AddSection(NAME_None);
+
+		TopSection.AddMenuEntry(
+			"CreateAttachCCW",
+			LOCTEXT("MenuCreateAttachCCW", "Add to New Color Correction Window"),
+			LOCTEXT("MenuCreateAttachCCW_Tooltip", "Creates new Color Correction Window (CCW) and adds valid selected actors to Per-Actor CC of the newly created CCW."),
+			FSlateIcon(),
+			FExecuteAction::CreateStatic(CreateNewCCActor, Window, InSelectedActors),
+			EUserInterfaceActionType::Button
+		);
 
-					+SVerticalBox::Slot()
-					.AutoHeight()
-					.Padding(0.0f, 9.0f, 0.0f, 0.0f)
+		TopSection.AddMenuEntry(
+			"CreateAttachCCR",
+			LOCTEXT("MenuCreateAttachCCR", "Add to New Color Correction Region"),
+			LOCTEXT("MenuCreateAttachCCR_Tooltip", "Creates new Color Correction Region (CCR) and adds valid selected actors to Per-Actor CC of the newly created CCR."),
+			FSlateIcon(),
+			FExecuteAction::CreateStatic(CreateNewCCActor, Region, InSelectedActors),
+			EUserInterfaceActionType::Button
+		);
+
+		FCreateSceneOutlinerMode ModeFactory = FCreateSceneOutlinerMode::CreateLambda([](SSceneOutliner* Outliner, TSharedPtr<TArray<AActor*>> SelectedActors)
+		{
+			return new FCCActorPickingMode(Outliner, FOnSceneOutlinerItemPicked::CreateStatic(OnAddPerActorCCFromActorTreeItemSelected, SelectedActors));
+		}, InSelectedActors);
+		
+		FSceneOutlinerInitializationOptions CCSceneOutlinerInitOptions;
+		CCSceneOutlinerInitOptions.bShowHeaderRow = false;
+		CCSceneOutlinerInitOptions.bFocusSearchBoxWhenOpened = true;
+		CCSceneOutlinerInitOptions.ModeFactory = ModeFactory;
+		CCSceneOutlinerInitOptions.Filters->AddFilterPredicate<FActorTreeItem>(FActorTreeItem::FFilterPredicate::CreateStatic(IsActorCCR));
+		
+		TSharedRef<SWidget> CCSceneOutliner =
+			SNew(SHorizontalBox)
+			+ SHorizontalBox::Slot()
+			[
+				SNew(SSceneOutliner, CCSceneOutlinerInitOptions)
+				.IsEnabled(FSlateApplication::Get().GetNormalExecutionAttribute())
+			]
+		
+			+ SHorizontalBox::Slot()
+			.VAlign(VAlign_Top)
+			.AutoWidth()
+			[
+				SNew(SVerticalBox)
+		
+				+ SVerticalBox::Slot()
+				.AutoHeight()
+				.Padding(0.0f, 9.0f, 0.0f, 0.0f)
+				[
+					SNew(SButton)
+					.ToolTipText(LOCTEXT("PickCCActorButtonLabel", "Pick a Color Correction Actor"))
+					.ButtonStyle(FAppStyle::Get(), "HoverHintOnly")
+					.OnClicked(FOnClicked::CreateStatic(PickCCActorMode, InSelectedActors))
+					.ForegroundColor(FSlateColor::UseForeground())
+					.IsFocusable(false)
 					[
-						SNew(SButton)
-						.ToolTipText( LOCTEXT( "PickCCActorButtonLabel", "Pick a Color Correction Actor") )
-						.ButtonStyle(FAppStyle::Get(), "HoverHintOnly")
-						.OnClicked(FOnClicked::CreateStatic(PickCCActorMode, InSelectedActors))
-						.ForegroundColor(FSlateColor::UseForeground())
-						.IsFocusable(false)
-						[
-							SNew(SImage)
-							.Image(FAppStyle::GetBrush("Icons.EyeDropper"))
-							.ColorAndOpacity(FSlateColor::UseForeground())
-						]
+						SNew(SImage)
+						.Image(FAppStyle::GetBrush("Icons.EyeDropper"))
+						.ColorAndOpacity(FSlateColor::UseForeground())
 					]
-				];
-			MenuBuilder.AddWidget(CCSceneOutliner, FText::GetEmpty(), /*bNoIndent=*/ true);
-		};
-
-		MenuBuilder.AddSubMenu(
-			LOCTEXT("CCRAttachToActorSubmenu", "Add to Per-Actor CC"),
-			LOCTEXT("CCRAttachToActorSubmenuToolTip", "Add currently selected actors to the list of actors that get affected by selected CC Actor in the following menu."),
-			FNewMenuDelegate::CreateLambda(CCActorAddToPerActorSubMenu, InSelectedActors),
-			false,
-			FSlateIcon()
+				]
+			];
+
+
+		FToolMenuSection& ExistingSection = Menu->AddSection("Existing", LOCTEXT("ExistingCCActorSection", "Existing"));
+		ExistingSection.AddEntry(
+			FToolMenuEntry::InitWidget("Picker", CCSceneOutliner, FText::GetEmpty(), /*bNoIndent=*/ true)
 		);
+	};
+
+	/**
+	 * Add the "Add to Per-Actor CC" menu entry, which either:
+	 * 1) When CCRs are selected, adds an actor from a picker menu to their per-actor CC; or
+	 * 2) When primitive actors are selected, adds them to the per-actor CC of a CCR from a picker menu
+	 */
+	static void AddAddToPerActorCCMenuEntry(FToolMenuSection& InSection, TSharedPtr<TArray<AActor*>> InSelectedActors)
+	{
+		const bool bIsCCROnlySelection = !InSelectedActors->ContainsByPredicate([](AActor* Actor)
+		{
+			return !IsActorCCR(Actor);
+		});
 
-		MenuBuilder.EndSection();
+		const bool bSelectionHasValidCCTargets = InSelectedActors->ContainsByPredicate(CanActorBeAddedToPerActorCC);
+
+		const FText EntryLabel = LOCTEXT("AddToPerActorCCEntryName", "Add to Per-Actor CC");
+
+		if (bIsCCROnlySelection)
+		{
+			InSection.AddSubMenu(
+				TEXT("AddToPerActorCCEntry_CCR"),
+				EntryLabel,
+				LOCTEXT("AddToPerActorCCEntryTooltip_CCR", "Add actor selected in the following menu to the list of actors that get affected by the currently selected CC Actors."),
+				FNewToolMenuDelegate::CreateStatic(CreateAddPerActorCCFromCCRMenu, InSelectedActors)
+			);
+		}
+		else if (bSelectionHasValidCCTargets)
+		{
+			InSection.AddSubMenu(
+				TEXT("AddToPerActorCCEntry_Actor"),
+				EntryLabel,
+				LOCTEXT("AddToPerActorCCEntryTooltip_Actor", "Add currently selected actors to the list of actors that get affected by selected CC Actor in the following menu."),
+				FNewToolMenuDelegate::CreateStatic(CreateAddPerActorCCFromActorMenu, InSelectedActors)
+			);
+		}
 	}
 
-	/** Creates a menu extender that is responsible for adding buttons and outliners for Per Actor CC. */
-	static TSharedRef<FExtender> ExtendContextMenu(const TSharedRef<FUICommandList> CommandList, TArray<AActor*> InSelectedActors)
+	/**
+	 * Add the "Remove Per-Actor CC" menu entry where appropriate
+	 */
+	static void AddRemoveFromPerActorCCMenuEntry(UToolMenu* InMenu, FToolMenuSection& InSection, TSharedPtr<TArray<AActor*>> InSelectedActors)
 	{
-		TSharedPtr<FExtender> Extender = MakeShared<FExtender>();
+		// This is only relevant when clicking in the Color Grading hierarchy
+		if (!InMenu->FindContext<UColorGradingMixerContextObject>())
+		{
+			return;
+		}
 
-		for (AActor* SelectedActor : InSelectedActors)
+		// This is also only relevant if the actor has a parent CCR in the hierarchy, and the parent affects the actor
+		UObjectMixerEditorListMenuContext* ObjectMixerContext = InMenu->FindContext<UObjectMixerEditorListMenuContext>();
+		if (!ObjectMixerContext)
 		{
-			for (UActorComponent* Component : SelectedActor->GetComponents())
+			return;
+		}
+
+		// Build a map from selected actor to CCRs from which the actor should be removed as an affected actor
+		TPerActorCCRemovalMap RemovalMap;
+
+		for (const TSharedPtr<ISceneOutlinerTreeItem>& TreeItem : ObjectMixerContext->Data.SelectedItems)
+		{
+			if (!TreeItem.IsValid() || !TreeItem->GetParent().IsValid())
 			{
-				if (UPrimitiveComponent* PrimitiveComponent = Cast<UPrimitiveComponent>(Component))
+				continue;
+			}
+
+			if (FActorTreeItem* ParentActorItem = TreeItem->GetParent()->CastTo<FActorTreeItem>())
+			{
+				if (AColorCorrectRegion* ParentCCR = Cast<AColorCorrectRegion>(ParentActorItem->Actor))
 				{
-					TSharedPtr<TArray<AActor*>> SelectedActorsPtr = MakeShared<TArray<AActor*>>(MoveTemp(InSelectedActors));
-					Extender->AddMenuExtension("ActorTypeTools", EExtensionHook::After, CommandList,
-						FMenuExtensionDelegate::CreateStatic(AddAttachToPerActorCCRMenu, SelectedActorsPtr));
-					return Extender.ToSharedRef();
+					for (AActor* SelectedActor : *InSelectedActors)
+					{
+						if (ParentCCR->AffectedActors.Contains(SelectedActor))
+						{
+							RemovalMap.FindOrAdd(SelectedActor).Add(ParentCCR);
+						}
+					}
 				}
 			}
 		}
-		return Extender.ToSharedRef();
+
+		if (RemovalMap.IsEmpty())
+		{
+			return;
+		}
+
+		const TSharedPtr<TPerActorCCRemovalMap> SharedRemovalMap = MakeShared<TPerActorCCRemovalMap>(MoveTemp(RemovalMap));
+
+		InSection.AddMenuEntry(
+			"MenuRemoveFromPerActorCC",
+			LOCTEXT("MenuRemoveFromPerActorCC", "Remove from Per-Actor CC"),
+			LOCTEXT("MenuRemoveFromPerActorCC_Tooltip", "Removes selected actors from the Per-Actor CC of the selection's parent CC Actors."),
+			FSlateIcon(),
+			FExecuteAction::CreateStatic(RemoveActorsFromPerActorCC, SharedRemovalMap),
+			EUserInterfaceActionType::Button
+		);
+	}
+
+	/** Adds the CCR section of a tool menu */
+	static void AddColorCorrectRegionsSection(UToolMenu* InMenu, TSharedPtr<TArray<AActor*>> InSelectedActors)
+	{
+		FToolMenuSection& Section = InMenu->FindOrAddSection(
+			"ColorCorrectionRegionsSection",
+			LOCTEXT("ColorCorrectionRegions", "Color Correction Regions"),
+			FToolMenuInsert("ActorTypeTools", EToolMenuInsertType::After)
+		);
+
+		AddAddToPerActorCCMenuEntry(Section, InSelectedActors);
+		AddRemoveFromPerActorCCMenuEntry(InMenu, Section, InSelectedActors);
 	}
 }
 
 void FColorCorrectionActorContextMenu::RegisterContextMenuExtender()
 {
-	CCRLevelViewportContextMenuExtender = FLevelEditorModule::FLevelViewportMenuExtender_SelectedActors::CreateStatic(ExtendContextMenu);
-	FLevelEditorModule& LevelEditorModule = FModuleManager::Get().LoadModuleChecked<FLevelEditorModule>("LevelEditor");
-	TArray<FLevelEditorModule::FLevelViewportMenuExtender_SelectedActors>& MenuExtenders = LevelEditorModule.GetAllLevelViewportContextMenuExtenders();
-	MenuExtenders.Add(CCRLevelViewportContextMenuExtender);
-	ContextMenuExtenderDelegateHandle = MenuExtenders.Last().GetHandle();
+	if (UToolMenu* ActorContextMenu = UToolMenus::Get()->ExtendMenu("LevelEditor.ActorContextMenu"))
+	{
+		FToolMenuSection& Section = ActorContextMenu->AddDynamicSection(NAME_None, FNewToolMenuDelegate::CreateLambda(
+			[](UToolMenu* InMenu)
+			{
+				if (ULevelEditorContextMenuContext* Context = InMenu->FindContext<ULevelEditorContextMenuContext>())
+				{
+					TArray<AActor*> SelectedActors;
+
+					if (Context->CurrentSelection)
+					{
+						Context->CurrentSelection->ForEachSelectedObject([&SelectedActors](UObject* Object)
+						{
+							if (AActor* Actor = Cast<AActor>(Object))
+							{
+								SelectedActors.Add(Actor);
+							}
+							return true;
+						});
+					}
+
+					AddColorCorrectRegionsSection(InMenu, MakeShared<TArray<AActor*>>(MoveTemp(SelectedActors)));
+				}
+			}),
+			FToolMenuInsert("ColorCorrectionRegionsSection", EToolMenuInsertType::After)
+		);
+	}
 }
 
 void FColorCorrectionActorContextMenu::UnregisterContextMenuExtender()
diff --git a/Engine/Plugins/Experimental/ColorCorrectRegions/Source/ColorCorrectRegionsEditor/Private/ColorGradingHierarchyConfig_ColorCorrectRegion.cpp b/Engine/Plugins/Experimental/ColorCorrectRegions/Source/ColorCorrectRegionsEditor/Private/ColorGradingHierarchyConfig_ColorCorrectRegion.cpp
--- a/Engine/Plugins/Experimental/ColorCorrectRegions/Source/ColorCorrectRegionsEditor/Private/ColorGradingHierarchyConfig_ColorCorrectRegion.cpp
+++ b/Engine/Plugins/Experimental/ColorCorrectRegions/Source/ColorCorrectRegionsEditor/Private/ColorGradingHierarchyConfig_ColorCorrectRegion.cpp
@@ -122,4 +122,12 @@ bool FColorGradingHierarchyConfig_ColorCorrectRegion::OnDrop(UObject* DropTarget
 	return false;
 }
 
+TSet<FName> FColorGradingHierarchyConfig_ColorCorrectRegion::GetPropertiesThatRequireListRefresh() const
+{
+	return {
+		GET_MEMBER_NAME_CHECKED(AColorCorrectionRegion, bEnablePerActorCC),
+		GET_MEMBER_NAME_CHECKED(AColorCorrectionRegion, AffectedActors)
+	};
+}
+
 #undef LOCTEXT_NAMESPACE
\ No newline at end of file
diff --git a/Engine/Plugins/Experimental/ColorCorrectRegions/Source/ColorCorrectRegionsEditor/Private/ColorGradingHierarchyConfig_ColorCorrectRegion.h b/Engine/Plugins/Experimental/ColorCorrectRegions/Source/ColorCorrectRegionsEditor/Private/ColorGradingHierarchyConfig_ColorCorrectRegion.h
--- a/Engine/Plugins/Experimental/ColorCorrectRegions/Source/ColorCorrectRegionsEditor/Private/ColorGradingHierarchyConfig_ColorCorrectRegion.h
+++ b/Engine/Plugins/Experimental/ColorCorrectRegions/Source/ColorCorrectRegionsEditor/Private/ColorGradingHierarchyConfig_ColorCorrectRegion.h
@@ -16,6 +16,6 @@ public:
 	virtual bool HasCustomDropHandling() const override;
 	virtual FSceneOutlinerDragValidationInfo ValidateDrop(UObject* DropTarget, const FSceneOutlinerDragDropPayload& Payload) const override;
 	virtual bool OnDrop(UObject* DropTarget, const FSceneOutlinerDragDropPayload& Payload, const FSceneOutlinerDragValidationInfo& ValidationInfo) const override;
-
+	virtual TSet<FName> GetPropertiesThatRequireListRefresh() const override;
 	//~ End IColorGradingMixerObjectHierarchyConfig interface
 };
\ No newline at end of file
