[ColorCorrectRegions] Fix duplicate LOCTEXT keys

diff --git a/Engine/Plugins/Experimental/ColorCorrectRegions/Source/ColorCorrectRegionsEditor/Private/ColorCorrectionActorContextMenu.cpp b/Engine/Plugins/Experimental/ColorCorrectRegions/Source/ColorCorrectRegionsEditor/Private/ColorCorrectionActorContextMenu.cpp
--- a/Engine/Plugins/Experimental/ColorCorrectRegions/Source/ColorCorrectRegionsEditor/Private/ColorCorrectionActorContextMenu.cpp
+++ b/Engine/Plugins/Experimental/ColorCorrectRegions/Source/ColorCorrectRegionsEditor/Private/ColorCorrectionActorContextMenu.cpp
@@ -82,7 +82,7 @@ namespace
 	/** Remove selected actors from the CC Actors' Per actor CC. */
 	static void RemoveActorsFromPerActorCC(TSharedPtr<TPerActorCCRemovalMap> RemovalMap)
 	{
-		FScopedTransaction Transaction(LOCTEXT("AddToPerActorCCTransaction", "Remove actors from Per-Actor CC"));
+		FScopedTransaction Transaction(LOCTEXT("RemoveFromPerActorCCTransaction", "Remove actors from Per-Actor CC"));
 
 		for (TPair<AActor*, TArray<AColorCorrectRegion*>> Pair : *RemovalMap)
 		{
