[ColorCorrectRegions] Fix CCR stencil failing to update when adding an actor from a context menu

diff --git a/Engine/Plugins/Experimental/ColorCorrectRegions/Source/ColorCorrectRegionsEditor/Private/ColorCorrectionActorContextMenu.cpp b/Engine/Plugins/Experimental/ColorCorrectRegions/Source/ColorCorrectRegionsEditor/Private/ColorCorrectionActorContextMenu.cpp
--- a/Engine/Plugins/Experimental/ColorCorrectRegions/Source/ColorCorrectRegionsEditor/Private/ColorCorrectionActorContextMenu.cpp
+++ b/Engine/Plugins/Experimental/ColorCorrectRegions/Source/ColorCorrectRegionsEditor/Private/ColorCorrectionActorContextMenu.cpp
@@ -67,14 +67,22 @@ namespace
 		{
 			if (AColorCorrectRegion* CCActorPtr = Cast<AColorCorrectRegion>(CCActor))
 			{
-				CCActorPtr->Modify();
+				// Forcing property change event. If we don't do this, the stencil doesn't update properly
+				const FName AffectedActorsPropertyName = GET_MEMBER_NAME_CHECKED(AColorCorrectRegion, AffectedActors);
+				FProperty* AffectedActorsProperty = FindFieldChecked<FProperty>(AColorCorrectRegion::StaticClass(), AffectedActorsPropertyName);
+				CCActorPtr->PreEditChange(AffectedActorsProperty);
 
 				CCActorPtr->bEnablePerActorCC = true;
+
 				for (AActor* SelectedActor : InAffectedActors)
 				{
 					TSoftObjectPtr<AActor> SelectedActorPtr(SelectedActor);
 					CCActorPtr->AffectedActors.Add(SelectedActorPtr);
 				}
+
+				FPropertyChangedEvent PropertyEvent(AffectedActorsProperty);
+				PropertyEvent.ChangeType = EPropertyChangeType::ArrayAdd;
+				CCActorPtr->PostEditChangeProperty(PropertyEvent);
 			}
 		}
 	}
@@ -91,7 +99,10 @@ namespace
 
 			for (AColorCorrectRegion* CCActor : CCActors)
 			{
-				CCActor->Modify();
+				// Forcing property change event. If we don't do this, the stencil doesn't update properly
+				const FName AffectedActorsPropertyName = GET_MEMBER_NAME_CHECKED(AColorCorrectRegion, AffectedActors);
+				FProperty* AffectedActorsProperty = FindFieldChecked<FProperty>(AColorCorrectRegion::StaticClass(), AffectedActorsPropertyName);
+				CCActor->PreEditChange(AffectedActorsProperty);
 
 				CCActor->AffectedActors.Remove(ActorToRemove);
 
@@ -100,6 +111,10 @@ namespace
 				{
 					CCActor->bEnablePerActorCC = false;
 				}
+
+				FPropertyChangedEvent PropertyEvent(AffectedActorsProperty);
+				PropertyEvent.ChangeType = EPropertyChangeType::ArrayRemove;
+				CCActor->PostEditChangeProperty(PropertyEvent);
 			}
 		}
 	}
