Color Correct Regions: Added transactions to context menu on actors when added to Per Actor CC.

#rnx

diff --git a/Engine/Plugins/Experimental/ColorCorrectRegions/Source/ColorCorrectRegionsEditor/Private/ColorCorrectionActorContextMenu.cpp b/Engine/Plugins/Experimental/ColorCorrectRegions/Source/ColorCorrectRegionsEditor/Private/ColorCorrectionActorContextMenu.cpp
--- a/Engine/Plugins/Experimental/ColorCorrectRegions/Source/ColorCorrectRegionsEditor/Private/ColorCorrectionActorContextMenu.cpp
+++ b/Engine/Plugins/Experimental/ColorCorrectRegions/Source/ColorCorrectRegionsEditor/Private/ColorCorrectionActorContextMenu.cpp
@@ -15,6 +15,7 @@
 #include "Editor/SceneOutliner/Public/SSceneOutliner.h"
 #include "Widgets/Input/SButton.h"
 #include "SSocketChooser.h"
+#include "ScopedTransaction.h"
 #include "Runtime/Engine/Classes/GameFramework/WorldSettings.h"
 #include "Editor/UnrealEdEngine.h"
 #include "UnrealEdGlobals.h"
@@ -61,7 +62,14 @@ namespace
 	/** Add Selected actors to the chosen CC Actors' Per actor CC. */
 	static void AddActorsToPerActorCC(TArray<AActor*> InCCActors, TArray<AActor*> InAffectedActors)
 	{
-		FScopedTransaction Transaction(LOCTEXT("AddToPerActorCCTransaction", "Add actors to Per-Actor CC"));
+		if (!InSelectedActors.IsValid() || InSelectedActors->Num() == 0)
+		{
+			return;
+		}
+
+#if WITH_EDITOR
+		const FScopedTransaction Transaction(LOCTEXT("PerActorCCActorListUpdated", "Added actors to Per actor CC via context menu."));
+#endif
 
 		for (AActor* CCActor : InCCActors)
 		{
@@ -85,6 +93,10 @@ namespace
 				CCActorPtr->PostEditChangeProperty(PropertyEvent);
 			}
 		}
+
+#if WITH_EDITOR
+		const FScopedTransaction Transaction(LOCTEXT("PerActorCCActorListUpdated", "Added actors to Per actor CC via context menu."));
+#endif
 	}
 
 	/** Remove selected actors from the CC Actors' Per actor CC. */
@@ -117,6 +129,9 @@ namespace
 				CCActor->PostEditChangeProperty(PropertyEvent);
 			}
 		}
+#if WITH_EDITOR
+		CCActorPtr->Modify();
+#endif
 	}
 
 	/** Creates a new CCR or CCW and then adds Selected actors to Per Actor CC. */
