[Backout] - CL37199548
Original CL Desc
-----------------------------------------------------------------
Color Correct Regions: Added transactions to context menu on actors when added to Per Actor CC.

#rnx

diff --git a/Engine/Plugins/Experimental/ColorCorrectRegions/Source/ColorCorrectRegionsEditor/Private/ColorCorrectionActorContextMenu.cpp b/Engine/Plugins/Experimental/ColorCorrectRegions/Source/ColorCorrectRegionsEditor/Private/ColorCorrectionActorContextMenu.cpp
--- a/Engine/Plugins/Experimental/ColorCorrectRegions/Source/ColorCorrectRegionsEditor/Private/ColorCorrectionActorContextMenu.cpp
+++ b/Engine/Plugins/Experimental/ColorCorrectRegions/Source/ColorCorrectRegionsEditor/Private/ColorCorrectionActorContextMenu.cpp
@@ -15,7 +15,6 @@
 #include "Editor/SceneOutliner/Public/SSceneOutliner.h"
 #include "Widgets/Input/SButton.h"
 #include "SSocketChooser.h"
-#include "ScopedTransaction.h"
 #include "Runtime/Engine/Classes/GameFramework/WorldSettings.h"
 #include "Editor/UnrealEdEngine.h"
 #include "UnrealEdGlobals.h"
@@ -62,14 +61,7 @@ namespace
 	/** Add Selected actors to the chosen CC Actors' Per actor CC. */
 	static void AddActorsToPerActorCC(TArray<AActor*> InCCActors, TArray<AActor*> InAffectedActors)
 	{
-		if (!InSelectedActors.IsValid() || InSelectedActors->Num() == 0)
-		{
-			return;
-		}
-
-#if WITH_EDITOR
-		const FScopedTransaction Transaction(LOCTEXT("PerActorCCActorListUpdated", "Added actors to Per actor CC via context menu."));
-#endif
+		FScopedTransaction Transaction(LOCTEXT("AddToPerActorCCTransaction", "Add actors to Per-Actor CC"));
 
 		for (AActor* CCActor : InCCActors)
 		{
@@ -93,10 +85,6 @@ namespace
 				CCActorPtr->PostEditChangeProperty(PropertyEvent);
 			}
 		}
-
-#if WITH_EDITOR
-		const FScopedTransaction Transaction(LOCTEXT("PerActorCCActorListUpdated", "Added actors to Per actor CC via context menu."));
-#endif
 	}
 
 	/** Remove selected actors from the CC Actors' Per actor CC. */
@@ -129,9 +117,6 @@ namespace
 				CCActor->PostEditChangeProperty(PropertyEvent);
 			}
 		}
-#if WITH_EDITOR
-		CCActorPtr->Modify();
-#endif
 	}
 
 	/** Creates a new CCR or CCW and then adds Selected actors to Per Actor CC. */
