Color Correct Regions. Storing pointers as weak object pointers. Somewhere between Tick and BeginRenderViewFamily something releases CC Actors making the pointers invalid.

#rnx

diff --git a/Engine/Plugins/Experimental/ColorCorrectRegions/Source/ColorCorrectRegions/Private/ColorCorrectRegionSubsystem.cpp b/Engine/Plugins/Experimental/ColorCorrectRegions/Source/ColorCorrectRegions/Private/ColorCorrectRegionSubsystem.cpp
--- a/Engine/Plugins/Experimental/ColorCorrectRegions/Source/ColorCorrectRegions/Private/ColorCorrectRegionSubsystem.cpp
+++ b/Engine/Plugins/Experimental/ColorCorrectRegions/Source/ColorCorrectRegions/Private/ColorCorrectRegionSubsystem.cpp
@@ -37,11 +37,11 @@ namespace
 #endif
 	}
 
-	void AssignNewPriorityIfNeeded(AColorCorrectRegion* InRegion, const TArray<AColorCorrectRegion*>& RegionsPriorityBased)
+	void AssignNewPriorityIfNeeded(AColorCorrectRegion* InRegion, const TArray<TWeakObjectPtr<AColorCorrectRegion>>& RegionsPriorityBased)
 	{
 		int32 HighestPriority = 0;
 		bool bAssignNewPriority = InRegion->Priority == 0;
-		for (const AColorCorrectRegion* Region : RegionsPriorityBased)
+		for (const TWeakObjectPtr<AColorCorrectRegion> Region : RegionsPriorityBased)
 		{
 			if (InRegion->Priority == Region->Priority)
 			{
@@ -59,12 +59,28 @@ namespace
 		}
 	}
 
-	FColorCorrectRenderProxyPtr CreateRenderStateForCCActor(AColorCorrectRegion* InActor)
+	FColorCorrectRenderProxyPtr CreateRenderStateForCCActor(TWeakObjectPtr<AColorCorrectRegion> InActorWeakPtr)
 	{
-		TRACE_CPUPROFILER_EVENT_SCOPE_TEXT(*FString::Printf(TEXT("CCR.TransferState %s"), *InActor->GetName()));
+		TStrongObjectPtr<AColorCorrectRegion> InActor = InActorWeakPtr.Pin();
+		if (!InActor.IsValid())
+		{
+			return nullptr;
+		}
 
 		FColorCorrectRenderProxyPtr TempCCRStateRenderThread = MakeShared<FColorCorrectRenderProxy>();
-		if (AColorCorrectionWindow* CCWindow = Cast<AColorCorrectionWindow>(InActor))
+
+		TempCCRStateRenderThread->bIsActiveThisFrame = InActor->Enabled
+#if WITH_EDITOR
+			&& !InActor->IsHiddenEd()
+#endif 
+			&& !(InActor->GetWorld()->HasBegunPlay() && InActor->IsHidden());
+
+		if (!TempCCRStateRenderThread->bIsActiveThisFrame)
+		{
+			return nullptr;
+		}
+
+		if (AColorCorrectionWindow* CCWindow = Cast<AColorCorrectionWindow>(InActor.Get()))
 		{
 			TempCCRStateRenderThread->WindowType = CCWindow->WindowType;
 			TempCCRStateRenderThread->ProxyType = FColorCorrectRenderProxy::DistanceBased;
@@ -75,13 +91,6 @@ namespace
 			TempCCRStateRenderThread->ProxyType = FColorCorrectRenderProxy::PriorityBased;
 		}
 
-		TempCCRStateRenderThread->bIsActiveThisFrame = IsValid(InActor)
-			&& InActor->Enabled
-#if WITH_EDITOR
-			&& !InActor->IsHiddenEd()
-#endif 
-			&& !(InActor->GetWorld()->HasBegunPlay() && InActor->IsHidden());
-
 		TempCCRStateRenderThread->World = InActor->GetWorld();
 		TempCCRStateRenderThread->Priority = InActor->Priority;
 		TempCCRStateRenderThread->Intensity = InActor->Intensity;
@@ -105,6 +114,7 @@ namespace
 		TempCCRStateRenderThread->PerActorColorCorrection = InActor->PerActorColorCorrection;
 
 		InActor->GetActorBounds(false, TempCCRStateRenderThread->BoxOrigin, TempCCRStateRenderThread->BoxExtent);
+
 		TempCCRStateRenderThread->ActorLocation = (FVector3f)InActor->GetActorLocation();
 		TempCCRStateRenderThread->ActorRotation = (FVector3f)InActor->GetActorRotation().Euler();
 		TempCCRStateRenderThread->ActorScale = (FVector3f)InActor->GetActorScale();
@@ -222,7 +232,6 @@ void UColorCorrectRegionsSubsystem::Deinitialize()
 void UColorCorrectRegionsSubsystem::Tick(float DeltaTime)
 {
 	Super::Tick(DeltaTime);
-
 	RefreshRegions();
 
 	// Check to make sure that no ids have been changed externally.
@@ -231,13 +240,19 @@ void UColorCorrectRegionsSubsystem::Tick(float DeltaTime)
 		const float WaitTimeInSecs = 1.0;
 		if (TimeSinceLastValidityCheck >= WaitTimeInSecs)
 		{
-			for (AColorCorrectRegion* Region : RegionsPriorityBased)
+			for (TWeakObjectPtr<AColorCorrectRegion> Region : RegionsPriorityBased)
 			{
-				CheckAssignedActorsValidity(Region);
+				if (Region.IsValid())
+				{
+					CheckAssignedActorsValidity(Region.Get());
+				}
 			}
-			for (AColorCorrectRegion* Region : RegionsDistanceBased)
+			for (TWeakObjectPtr<AColorCorrectRegion> Region : RegionsDistanceBased)
 			{
-				CheckAssignedActorsValidity(Region);
+				if (Region.IsValid())
+				{
+					CheckAssignedActorsValidity(Region.Get());
+				}
 			}
 
 			TimeSinceLastValidityCheck = 0;
@@ -248,16 +263,24 @@ void UColorCorrectRegionsSubsystem::Tick(float DeltaTime)
 
 void UColorCorrectRegionsSubsystem::TransferStates()
 {
+	TRACE_CPUPROFILER_EVENT_SCOPE(TEXT("CCR.TransferStates"));
+
 	TArray<FColorCorrectRenderProxyPtr> TempProxiesPriority;
 	TArray<FColorCorrectRenderProxyPtr> TempProxiesDistance;
-	for (AColorCorrectRegion* Region : RegionsPriorityBased)
+	for (TWeakObjectPtr<AColorCorrectRegion> Region : RegionsPriorityBased)
 	{
-		TempProxiesPriority.Add(CreateRenderStateForCCActor(Region));
+		if (FColorCorrectRenderProxyPtr ProxyPtr = CreateRenderStateForCCActor(Region))
+		{
+			TempProxiesPriority.Add(ProxyPtr);
+		}
 	}
 
-	for (AColorCorrectRegion* Region : RegionsDistanceBased)
+	for (TWeakObjectPtr<AColorCorrectRegion> Region : RegionsDistanceBased)
 	{
-		TempProxiesDistance.Add(CreateRenderStateForCCActor(Region));
+		if (FColorCorrectRenderProxyPtr ProxyPtr = CreateRenderStateForCCActor(Region))
+		{
+			TempProxiesDistance.Add(ProxyPtr);
+		}
 	}
 
 	// Sort priority based proxies on game thread.
diff --git a/Engine/Plugins/Experimental/ColorCorrectRegions/Source/ColorCorrectRegions/Public/ColorCorrectRegionsSubsystem.h b/Engine/Plugins/Experimental/ColorCorrectRegions/Source/ColorCorrectRegions/Public/ColorCorrectRegionsSubsystem.h
--- a/Engine/Plugins/Experimental/ColorCorrectRegions/Source/ColorCorrectRegions/Public/ColorCorrectRegionsSubsystem.h
+++ b/Engine/Plugins/Experimental/ColorCorrectRegions/Source/ColorCorrectRegions/Public/ColorCorrectRegionsSubsystem.h
@@ -106,9 +106,9 @@ private:
 private:
 
 	/** Stores pointers to ColorCorrectRegion Actors that use priority for sorting. */
-	TArray<AColorCorrectRegion*> RegionsPriorityBased;
+	TArray<TWeakObjectPtr<AColorCorrectRegion>> RegionsPriorityBased;
 	/** Stores pointers to ColorCorrectRegion Actors that are based on distance from camera. */
-	TArray<AColorCorrectRegion*> RegionsDistanceBased;
+	TArray<TWeakObjectPtr<AColorCorrectRegion>> RegionsDistanceBased;
 
 
 	/** Proxies to be used exclusively on render thread. Copies of the state of CC Actors sorted by priority. */
