Img Media/Color Correction Actors: When CC Actors are out of bounds the viewport.Scale method could cause a narrowing conversion crash because Scale does widening conversion and then checked narrowing conversion.
The fix is to operate on widened 64 bit TIntRect and do a checkless narrowing conversion when viewport checks have passed.
For Img Media we output an error to the log instead of hard crashing.

[RN] Fixed a crash in Color Correction Actors caused by large resolutions and CCRs that far outside of the view.

diff --git a/Engine/Plugins/Experimental/ColorCorrectRegions/Source/ColorCorrectRegions/Private/ColorCorrectRegionsSceneViewExtension.cpp b/Engine/Plugins/Experimental/ColorCorrectRegions/Source/ColorCorrectRegions/Private/ColorCorrectRegionsSceneViewExtension.cpp
--- a/Engine/Plugins/Experimental/ColorCorrectRegions/Source/ColorCorrectRegions/Private/ColorCorrectRegionsSceneViewExtension.cpp
+++ b/Engine/Plugins/Experimental/ColorCorrectRegions/Source/ColorCorrectRegions/Private/ColorCorrectRegionsSceneViewExtension.cpp
@@ -359,7 +359,7 @@ namespace
 			return false;
 		}
 
-		FIntRect Viewport;
+		FIntRect BoundingRectangle;
 
 		float MaxDepth = -BIG_NUMBER;
 		float MinDepth = BIG_NUMBER;
@@ -367,14 +367,14 @@ namespace
 		if (RegionState->Invert)
 		{
 			// In case of Region inversion we would to render the entire screen
-			Viewport = PrimaryViewRect;
+			BoundingRectangle = PrimaryViewRect;
 		}
 		else
 		{
-			GetPixelSpaceBoundingRect(View, RegionState->BoxOrigin, RegionState->BoxExtent, Viewport, MaxDepth, MinDepth);
+			GetPixelSpaceBoundingRect(View, RegionState->BoxOrigin, RegionState->BoxExtent, BoundingRectangle, MaxDepth, MinDepth);
 
 			// Check if CCR is too small to be rendered (less than one pixel on the screen).
-			if (Viewport.Width() == 0 || Viewport.Height() == 0)
+			if (BoundingRectangle.Width() == 0 || BoundingRectangle.Height() == 0)
 			{
 				return false;
 			}
@@ -384,28 +384,37 @@ namespace
 			float MaxBoxExtent = FMath::Abs(RegionState->BoxExtent.GetMax());
 			if (MaxDepth >= 0 && MinDepth < 0)
 			{
-				UpdateMinMaxWithFrustrumAABBIntersection(View, RegionState->BoxOrigin, RegionState->BoxExtent, Viewport, MaxDepth);
+				UpdateMinMaxWithFrustrumAABBIntersection(View, RegionState->BoxOrigin, RegionState->BoxExtent, BoundingRectangle, MaxDepth);
 			}
 
 			FIntRect ConstrainedViewRect = View.UnscaledViewRect;
 
 			// We need to make sure that Bounding Rectangle is offset by the position of the View's Viewport.
-			Viewport.Min -= ConstrainedViewRect.Min;
+			BoundingRectangle.Min -= ConstrainedViewRect.Min;
+			UE::Math::TIntRect<int64> BoundingRectToBeTruncated(BoundingRectangle);
 
-			Viewport = Viewport.Scale(ScreenPercentage);
+			// As of UE5.5 Scale method has non templated 64 bit narrowing conversion check even if FIntRect is 32 bit. 
+			BoundingRectToBeTruncated = BoundingRectToBeTruncated.Scale(ScreenPercentage);
 
-			// Culling all regions that are not within the screen bounds.
-			if ((Viewport.Min.X >= PrimaryViewRect.Width() ||
-				Viewport.Min.Y >= PrimaryViewRect.Height() ||
-				Viewport.Max.X <= 0.0f ||
-				Viewport.Max.Y <= 0.0f ||
+			// Culling all regions that are not within the screen bounds on a 64 bit viewport.
+			if ((BoundingRectToBeTruncated.Min.X >= PrimaryViewRect.Width() ||
+				BoundingRectToBeTruncated.Min.Y >= PrimaryViewRect.Height() ||
+				BoundingRectToBeTruncated.Max.X <= 0.0f ||
+				BoundingRectToBeTruncated.Max.Y <= 0.0f ||
 				MaxDepth < 0.0f))
 			{
 				return false;
 			}
+
+			// Truncate back to 32 bit. At this point CC Actors with bounding rectangles that do not fit the viewport are skipped.
+			BoundingRectangle = FIntRect(
+				FIntPoint((int32)BoundingRectToBeTruncated.Min.X, (int32)BoundingRectToBeTruncated.Min.Y),
+				FIntPoint((int32)BoundingRectToBeTruncated.Max.X, (int32)BoundingRectToBeTruncated.Max.Y)
+				);
+
 			// Clipping is required because as we get closer to the bounding box the bounds
 			// May extend beyond Allowed render target size.
-			Viewport.Clip(PrimaryViewRect);
+			BoundingRectangle.Clip(PrimaryViewRect);
 		}
 
 
@@ -432,7 +441,7 @@ namespace
 			}
 		}
 
-		const FScreenPassTextureViewport RegionViewport(SceneColorRenderTarget.Texture, Viewport);
+		const FScreenPassTextureViewport RegionViewport(SceneColorRenderTarget.Texture, BoundingRectangle);
 
 		FCCRShaderInputParameters* PostProcessMaterialParameters = GraphBuilder.AllocParameters<FCCRShaderInputParameters>();
 		PostProcessMaterialParameters->RenderTargets[0] = BackBufferRenderTarget.GetRenderTargetBinding();
diff --git a/Engine/Plugins/Media/ImgMedia/Source/ImgMedia/Private/ImgMediaSceneViewExtension.cpp b/Engine/Plugins/Media/ImgMedia/Source/ImgMedia/Private/ImgMediaSceneViewExtension.cpp
--- a/Engine/Plugins/Media/ImgMedia/Source/ImgMedia/Private/ImgMediaSceneViewExtension.cpp
+++ b/Engine/Plugins/Media/ImgMedia/Source/ImgMedia/Private/ImgMediaSceneViewExtension.cpp
@@ -1,6 +1,7 @@
 // Copyright Epic Games, Inc. All Rights Reserved.
 
 #include "ImgMediaSceneViewExtension.h"
+#include "ImgMediaPrivate.h"
 #include "DynamicResolutionState.h"
 #include "Framework/Notifications/NotificationManager.h"
 #include "Widgets/Notifications/SNotificationList.h"
@@ -90,7 +91,33 @@ void FImgMediaSceneViewExtension::CacheViewInfo(FSceneViewFamily& InViewFamily,
 	Info.Location = View.ViewMatrices.GetViewOrigin();
 	Info.ViewDirection = View.GetViewDirection();
 	Info.ViewProjectionMatrix = View.ViewMatrices.GetViewProjectionMatrix();
-	Info.ViewportRect = View.UnconstrainedViewRect.Scale(ResolutionFraction);
+	
+	// This converts a potential narrowing conversion crash to an error. 
+	// If FIntRect does 32 to 64 bit widening and then checked narrowing from 64 to 32, which could cause a crash
+	// if the value exceedes the 32 bit int limit. 
+	{
+		UE::Math::TIntRect<int64> UnconstrainedViewportRect64(View.UnconstrainedViewRect);
+		UnconstrainedViewportRect64 = UnconstrainedViewportRect64.Scale(ResolutionFraction);
+
+		if (!IntFitsIn<int32>(UnconstrainedViewportRect64.Min.X) ||
+			!IntFitsIn<int32>(UnconstrainedViewportRect64.Min.Y) ||
+			!IntFitsIn<int32>(UnconstrainedViewportRect64.Max.X) ||
+			!IntFitsIn<int32>(UnconstrainedViewportRect64.Max.Y))
+		{
+			UE_LOG(LogImgMedia, Error, TEXT("Scaled Unconstrained viewport is out of bounds. Original Viewport rect: Min: %d x %d, Max: %d x %d, Screen Percentage: %f"), 
+				View.UnconstrainedViewRect.Min.X,
+				View.UnconstrainedViewRect.Min.Y,
+				View.UnconstrainedViewRect.Max.X,
+				View.UnconstrainedViewRect.Max.Y,
+				ResolutionFraction);
+		}
+
+		Info.ViewportRect = FIntRect(
+			FIntPoint((int32)UnconstrainedViewportRect64.Min.X, (int32)UnconstrainedViewportRect64.Min.Y),
+			FIntPoint((int32)UnconstrainedViewportRect64.Max.X, (int32)UnconstrainedViewportRect64.Max.Y)
+		);
+	}
+
 
 	if (FMath::IsNearlyEqual(FieldOfViewMultiplier, 1.0f))
 	{
