Fixed usage of TRACE_CPUPROFILER_EVENT_SCOPE.

#rnx

diff --git a/Engine/Plugins/Experimental/ColorCorrectRegions/Source/ColorCorrectRegions/Private/ColorCorrectRegionSubsystem.cpp b/Engine/Plugins/Experimental/ColorCorrectRegions/Source/ColorCorrectRegions/Private/ColorCorrectRegionSubsystem.cpp
--- a/Engine/Plugins/Experimental/ColorCorrectRegions/Source/ColorCorrectRegions/Private/ColorCorrectRegionSubsystem.cpp
+++ b/Engine/Plugins/Experimental/ColorCorrectRegions/Source/ColorCorrectRegions/Private/ColorCorrectRegionSubsystem.cpp
@@ -263,7 +263,7 @@ void UColorCorrectRegionsSubsystem::Tick(float DeltaTime)
 
 void UColorCorrectRegionsSubsystem::TransferStates()
 {
-	TRACE_CPUPROFILER_EVENT_SCOPE(TEXT("CCR.TransferStates"));
+	TRACE_CPUPROFILER_EVENT_SCOPE(CCR.TransferStates);
 
 	TArray<FColorCorrectRenderProxyPtr> TempProxiesPriority;
 	TArray<FColorCorrectRenderProxyPtr> TempProxiesDistance;
