Color Correction Regions: Fix for the following issue.
There is a case where world on the actor is null, while the actor itself is considered to be valid. This happens when world partition is used.

diff --git a/Engine/Plugins/Experimental/ColorCorrectRegions/Source/ColorCorrectRegions/Private/ColorCorrectRegionSubsystem.cpp b/Engine/Plugins/Experimental/ColorCorrectRegions/Source/ColorCorrectRegions/Private/ColorCorrectRegionSubsystem.cpp
--- a/Engine/Plugins/Experimental/ColorCorrectRegions/Source/ColorCorrectRegions/Private/ColorCorrectRegionSubsystem.cpp
+++ b/Engine/Plugins/Experimental/ColorCorrectRegions/Source/ColorCorrectRegions/Private/ColorCorrectRegionSubsystem.cpp
@@ -62,7 +62,7 @@ namespace
 	FColorCorrectRenderProxyPtr CreateRenderStateForCCActor(TWeakObjectPtr<AColorCorrectRegion> InActorWeakPtr)
 	{
 		TStrongObjectPtr<AColorCorrectRegion> InActor = InActorWeakPtr.Pin();
-		if (!InActor.IsValid())
+		if (!InActor.IsValid() || !IsValid(InActor->GetWorld()))
 		{
 			return nullptr;
 		}
