Color Correction Actors: Added more prominence to the invalid stencil setting in project settings.

#rnx

diff --git a/Engine/Plugins/Experimental/ColorCorrectRegions/Source/ColorCorrectRegions/ColorCorrectRegions.Build.cs b/Engine/Plugins/Experimental/ColorCorrectRegions/Source/ColorCorrectRegions/ColorCorrectRegions.Build.cs
--- a/Engine/Plugins/Experimental/ColorCorrectRegions/Source/ColorCorrectRegions/ColorCorrectRegions.Build.cs
+++ b/Engine/Plugins/Experimental/ColorCorrectRegions/Source/ColorCorrectRegions/ColorCorrectRegions.Build.cs
@@ -13,7 +13,7 @@ public class ColorCorrectRegions : ModuleRules
 				"DisplayClusterLightCardExtender"
 			}
 		);
-			
+
 		PrivateDependencyModuleNames.AddRange(
 			new string[]
 			{
@@ -31,6 +31,8 @@ public class ColorCorrectRegions : ModuleRules
 		{
 			PrivateDependencyModuleNames.Add("UnrealEd");
 			PrivateDependencyModuleNames.Add("EditorWidgets");
+			PrivateDependencyModuleNames.Add("Slate");
+			PrivateDependencyModuleNames.Add("SlateCore");
 		}
 	}
 }
diff --git a/Engine/Plugins/Experimental/ColorCorrectRegions/Source/ColorCorrectRegions/Private/ColorCorrectRegionSubsystem.cpp b/Engine/Plugins/Experimental/ColorCorrectRegions/Source/ColorCorrectRegions/Private/ColorCorrectRegionSubsystem.cpp
--- a/Engine/Plugins/Experimental/ColorCorrectRegions/Source/ColorCorrectRegions/Private/ColorCorrectRegionSubsystem.cpp
+++ b/Engine/Plugins/Experimental/ColorCorrectRegions/Source/ColorCorrectRegions/Private/ColorCorrectRegionSubsystem.cpp
@@ -13,7 +13,9 @@
 
 #if WITH_EDITOR
 #include "Editor.h"
+#include "Framework/Notifications/NotificationManager.h"
 #include "ScopedTransaction.h"
+#include "Widgets/Notifications/SNotificationList.h"
 #endif
 
 #define LOCTEXT_NAMESPACE "CCR"
@@ -59,7 +61,7 @@ namespace
 		}
 	}
 
-	FColorCorrectRenderProxyPtr CreateRenderStateForCCActor(TWeakObjectPtr<AColorCorrectRegion> InActorWeakPtr)
+	FColorCorrectRenderProxyPtr CreateRenderStateForCCActor(TWeakObjectPtr<AColorCorrectRegion> InActorWeakPtr, const bool bInSupportsStencil)
 	{
 		TStrongObjectPtr<AColorCorrectRegion> InActor = InActorWeakPtr.Pin();
 		if (!InActor.IsValid() || !IsValid(InActor->GetWorld()))
@@ -110,7 +112,7 @@ namespace
 		TempCCRStateRenderThread->Temperature = InActor->Temperature;
 		TempCCRStateRenderThread->Tint = InActor->Tint;
 		TempCCRStateRenderThread->ColorGradingSettings = InActor->ColorGradingSettings;
-		TempCCRStateRenderThread->bEnablePerActorCC = InActor->bEnablePerActorCC;
+		TempCCRStateRenderThread->bEnablePerActorCC = bInSupportsStencil && InActor->bEnablePerActorCC;
 		TempCCRStateRenderThread->PerActorColorCorrection = InActor->PerActorColorCorrection;
 
 		InActor->GetActorBounds(false, TempCCRStateRenderThread->BoxOrigin, TempCCRStateRenderThread->BoxExtent);
@@ -267,9 +269,32 @@ void UColorCorrectRegionsSubsystem::TransferStates()
 
 	TArray<FColorCorrectRenderProxyPtr> TempProxiesPriority;
 	TArray<FColorCorrectRenderProxyPtr> TempProxiesDistance;
+
+	// Custom Depth required to be set to Enabled with stencil for Per Actor CC feature to work.
+	static bool bNotifiedAboutCustomDepth = false;
+	static const auto CVarCustomDepth = IConsoleManager::Get().FindTConsoleVariableDataInt(TEXT("r.CustomDepth"));
+	constexpr int32 EnabledWithStencil = 3;
+	const bool bSupportsStencil = CVarCustomDepth->GetValueOnAnyThread() == EnabledWithStencil;
+
+	if (!bNotifiedAboutCustomDepth && !bSupportsStencil)
+	{
+		FString InvalidCustomDepthSettingString = "Per Actor Color Correction requires Custom Depth Mode to be set to \"Enabled With Stencil\"";
+		UE_LOG(ColorCorrectRegions, Error, TEXT("%s"), *InvalidCustomDepthSettingString);
+
+#if WITH_EDITOR
+		FNotificationInfo Info(FText::FromString(InvalidCustomDepthSettingString));
+		Info.ExpireDuration = 5.0f;
+		Info.Image = FCoreStyle::Get().GetBrush(TEXT("MessageLog.Warning"));
+
+		FSlateNotificationManager::Get().AddNotification(Info);
+#endif
+
+		bNotifiedAboutCustomDepth = true;
+	}
+
 	for (TWeakObjectPtr<AColorCorrectRegion> Region : RegionsPriorityBased)
 	{
-		if (FColorCorrectRenderProxyPtr ProxyPtr = CreateRenderStateForCCActor(Region))
+		if (FColorCorrectRenderProxyPtr ProxyPtr = CreateRenderStateForCCActor(Region, bSupportsStencil))
 		{
 			TempProxiesPriority.Add(ProxyPtr);
 		}
@@ -277,7 +302,7 @@ void UColorCorrectRegionsSubsystem::TransferStates()
 
 	for (TWeakObjectPtr<AColorCorrectRegion> Region : RegionsDistanceBased)
 	{
-		if (FColorCorrectRenderProxyPtr ProxyPtr = CreateRenderStateForCCActor(Region))
+		if (FColorCorrectRenderProxyPtr ProxyPtr = CreateRenderStateForCCActor(Region, bSupportsStencil))
 		{
 			TempProxiesDistance.Add(ProxyPtr);
 		}
diff --git a/Engine/Plugins/Experimental/ColorCorrectRegions/Source/ColorCorrectRegions/Private/ColorCorrectRegionsSceneViewExtension.cpp b/Engine/Plugins/Experimental/ColorCorrectRegions/Source/ColorCorrectRegions/Private/ColorCorrectRegionsSceneViewExtension.cpp
--- a/Engine/Plugins/Experimental/ColorCorrectRegions/Source/ColorCorrectRegions/Private/ColorCorrectRegionsSceneViewExtension.cpp
+++ b/Engine/Plugins/Experimental/ColorCorrectRegions/Source/ColorCorrectRegions/Private/ColorCorrectRegionsSceneViewExtension.cpp
@@ -247,21 +247,6 @@ namespace
 		, const TArray<uint32>& StencilIds
 		, FScreenPassRenderTarget& OutMergedStencilRenderTarget)
 	{
-		static bool bNotifiedAboutCustomDepth = false;
-		static const auto CVarCustomDepth = IConsoleManager::Get().FindTConsoleVariableDataInt(TEXT("r.CustomDepth"));
-		const int32 EnabledWithStencil = 3;
-
-		if (CVarCustomDepth->GetValueOnAnyThread() != EnabledWithStencil && !bNotifiedAboutCustomDepth)
-		{
-			UE_LOG(ColorCorrectRegions, Error, TEXT("Per Actor Color Correction requires Custom Depth Mode to be set to \"Enabled With Stencil\""));
-			bNotifiedAboutCustomDepth = true;
-			return;
-		}
-		else if (CVarCustomDepth->GetValueOnAnyThread() == EnabledWithStencil)
-		{
-			bNotifiedAboutCustomDepth = false;
-		}
-
 		if (StencilIds.Num() == 0)
 		{
 			return;
