CCR: Only emit warning when CCRs are in use.

diff --git a/Engine/Plugins/Experimental/ColorCorrectRegions/Source/ColorCorrectRegions/Private/ColorCorrectRegionSubsystem.cpp b/Engine/Plugins/Experimental/ColorCorrectRegions/Source/ColorCorrectRegions/Private/ColorCorrectRegionSubsystem.cpp
--- a/Engine/Plugins/Experimental/ColorCorrectRegions/Source/ColorCorrectRegions/Private/ColorCorrectRegionSubsystem.cpp
+++ b/Engine/Plugins/Experimental/ColorCorrectRegions/Source/ColorCorrectRegions/Private/ColorCorrectRegionSubsystem.cpp
@@ -276,22 +276,6 @@ void UColorCorrectRegionsSubsystem::TransferStates()
 	constexpr int32 EnabledWithStencil = 3;
 	const bool bSupportsStencil = CVarCustomDepth->GetValueOnAnyThread() == EnabledWithStencil;
 
-	if (!bNotifiedAboutCustomDepth && !bSupportsStencil)
-	{
-		FString InvalidCustomDepthSettingString = "Per Actor Color Correction requires Custom Depth Mode to be set to \"Enabled With Stencil\"";
-		UE_LOG(ColorCorrectRegions, Error, TEXT("%s"), *InvalidCustomDepthSettingString);
-
-#if WITH_EDITOR
-		FNotificationInfo Info(FText::FromString(InvalidCustomDepthSettingString));
-		Info.ExpireDuration = 5.0f;
-		Info.Image = FCoreStyle::Get().GetBrush(TEXT("MessageLog.Warning"));
-
-		FSlateNotificationManager::Get().AddNotification(Info);
-#endif
-
-		bNotifiedAboutCustomDepth = true;
-	}
-
 	for (TWeakObjectPtr<AColorCorrectRegion> Region : RegionsPriorityBased)
 	{
 		if (FColorCorrectRenderProxyPtr ProxyPtr = CreateRenderStateForCCActor(Region, bSupportsStencil))
@@ -314,6 +298,26 @@ void UColorCorrectRegionsSubsystem::TransferStates()
 		return A->Priority < B->Priority;
 	});
 
+	const int32 TransferCount = TempProxiesPriority.Num() + TempProxiesDistance.Num();
+	if (TransferCount > 0)
+	{
+		if (!bNotifiedAboutCustomDepth && !bSupportsStencil)
+		{
+			FString InvalidCustomDepthSettingString = "Per Actor Color Correction requires Custom Depth Mode to be set to \"Enabled With Stencil\"";
+			UE_LOG(ColorCorrectRegions, Error, TEXT("%s"), *InvalidCustomDepthSettingString);
+
+#if WITH_EDITOR
+			FNotificationInfo Info(FText::FromString(InvalidCustomDepthSettingString));
+			Info.ExpireDuration = 5.0f;
+			Info.Image = FCoreStyle::Get().GetBrush(TEXT("MessageLog.Warning"));
+
+			FSlateNotificationManager::Get().AddNotification(Info);
+#endif
+
+			bNotifiedAboutCustomDepth = true;
+		}
+	}
+
 	ENQUEUE_RENDER_COMMAND(CopyCCProxies)([this, TempProxiesPriority = MoveTemp(TempProxiesPriority), TempProxiesDistance = MoveTemp(TempProxiesDistance)](FRHICommandListImmediate& RHICmdList)
 	{
 		ProxiesPriorityBased = TempProxiesPriority;
diff --git a/Engine/Plugins/Experimental/ColorCorrectRegions/Source/ColorCorrectRegions/Private/ColorCorrectRegionsSceneViewExtension.cpp b/Engine/Plugins/Experimental/ColorCorrectRegions/Source/ColorCorrectRegions/Private/ColorCorrectRegionsSceneViewExtension.cpp
--- a/Engine/Plugins/Experimental/ColorCorrectRegions/Source/ColorCorrectRegions/Private/ColorCorrectRegionsSceneViewExtension.cpp
+++ b/Engine/Plugins/Experimental/ColorCorrectRegions/Source/ColorCorrectRegions/Private/ColorCorrectRegionsSceneViewExtension.cpp
@@ -657,7 +657,10 @@ FColorCorrectRegionsSceneViewExtension::FColorCorrectRegionsSceneViewExtension(c
 
 void FColorCorrectRegionsSceneViewExtension::BeginRenderViewFamily(FSceneViewFamily& InViewFamily)
 {
-	WorldSubsystem->TransferStates();
+	if (IsValid(WorldSubsystem))
+	{
+		WorldSubsystem->TransferStates();
+	}
 }
 
 void FColorCorrectRegionsSceneViewExtension::PrePostProcessPass_RenderThread(FRDGBuilder& GraphBuilder, const FSceneView& View, const FPostProcessingInputs& Inputs)
