Color Correction Regions: Clamping Region viewport to max viewport bounds. There is state that CCRs get into after a very long run under unknown conditions that cause for Bounding rectangle to exceed the bounds.

#rnx

diff --git a/Engine/Plugins/Experimental/ColorCorrectRegions/Source/ColorCorrectRegions/Private/ColorCorrectRegionsSceneViewExtension.cpp b/Engine/Plugins/Experimental/ColorCorrectRegions/Source/ColorCorrectRegions/Private/ColorCorrectRegionsSceneViewExtension.cpp
--- a/Engine/Plugins/Experimental/ColorCorrectRegions/Source/ColorCorrectRegions/Private/ColorCorrectRegionsSceneViewExtension.cpp
+++ b/Engine/Plugins/Experimental/ColorCorrectRegions/Source/ColorCorrectRegions/Private/ColorCorrectRegionsSceneViewExtension.cpp
@@ -65,7 +65,11 @@ namespace
 
 	void GetPixelSpaceBoundingRect(const FSceneView& InView, const FVector& InBoxCenter, const FVector& InBoxExtents, FIntRect& OutViewport, float& OutMaxDepth, float& OutMinDepth)
 	{
-		OutViewport = FIntRect(INT32_MAX, INT32_MAX, -INT32_MAX, -INT32_MAX);
+		// Equivalent to D3D12_VIEWPORT_BOUNDS_MAX
+		constexpr int32 CMaxViewportDimension = 32767;
+
+		OutViewport = FIntRect(CMaxViewportDimension, CMaxViewportDimension, -CMaxViewportDimension, -CMaxViewportDimension);
+
 		// 8 corners of the bounding box. To be multiplied by box extent and offset by the center.
 		const int NumCorners = 8;
 		const FVector Verts[NumCorners] = {
@@ -356,12 +360,6 @@ namespace
 		{
 			GetPixelSpaceBoundingRect(View, RegionState->BoxOrigin, RegionState->BoxExtent, BoundingRectangle, MaxDepth, MinDepth);
 
-			// Check if CCR is too small to be rendered (less than one pixel on the screen).
-			if (BoundingRectangle.Width() == 0 || BoundingRectangle.Height() == 0)
-			{
-				return false;
-			}
-
 			// This is to handle corner cases when user has a very long disproportionate region and gets either
 			// within bounds or close to the center.
 			float MaxBoxExtent = FMath::Abs(RegionState->BoxExtent.GetMax());
@@ -398,6 +396,24 @@ namespace
 			// Clipping is required because as we get closer to the bounding box the bounds
 			// May extend beyond Allowed render target size.
 			BoundingRectangle.Clip(PrimaryViewRect);
+
+			// Equivalent to D3D12_VIEWPORT_BOUNDS_MAX
+			constexpr int32 CMaxViewportDimension = 32767;
+
+			// Equivalent to D3D12_VIEWPORT_BOUNDS_MIN
+			constexpr int32 CMinViewportDimension = -32768;
+
+			// Clamp the bounding rectangle to the max viewport dimensions.
+			BoundingRectangle.Min.X = FMath::Clamp(BoundingRectangle.Min.X, CMinViewportDimension, CMaxViewportDimension);
+			BoundingRectangle.Min.Y = FMath::Clamp(BoundingRectangle.Min.Y, CMinViewportDimension, CMaxViewportDimension);
+			BoundingRectangle.Max.X = FMath::Clamp(BoundingRectangle.Max.X, CMinViewportDimension, CMaxViewportDimension);
+			BoundingRectangle.Max.Y = FMath::Clamp(BoundingRectangle.Max.Y, CMinViewportDimension, CMaxViewportDimension);
+
+			// Check if CCR is too small to be rendered (less than one pixel on the screen).
+			if (BoundingRectangle.Width() == 0 || BoundingRectangle.Height() == 0)
+			{
+				return false;
+			}
 		}
 
 
