Color Correct Actors: Moved some public headers into private folders to avoid unnecessary backwards compatibility support. Also removed some old deprecated methods.

#rnx

diff --git a/Engine/Plugins/Experimental/ColorCorrectRegions/Source/ColorCorrectRegions/Private/ColorCorrectRegionsPostProcessMaterial.h b/Engine/Plugins/Experimental/ColorCorrectRegions/Source/ColorCorrectRegions/Private/ColorCorrectRegionsPostProcessMaterial.h
new file mode 100644
--- /dev/null
+++ b/Engine/Plugins/Experimental/ColorCorrectRegions/Source/ColorCorrectRegions/Private/ColorCorrectRegionsPostProcessMaterial.h
@@ -0,0 +1,248 @@
+// Copyright Epic Games, Inc. All Rights Reserved.
+
+#pragma once
+#include "ColorCorrectRegion.h"
+#include "ColorCorrectWindow.h"
+#include "CoreMinimal.h"
+#include "DataDrivenShaderPlatformInfo.h"
+#include "MaterialShader.h"
+#include "Materials/Material.h"
+#include "Materials/MaterialInstanceDynamic.h"
+#include "Materials/MaterialInterface.h"
+#include "RHI.h"
+#include "RHIStaticStates.h"
+#include "RenderGraphResources.h"
+#include "RenderGraphResources.h"
+#include "RenderResource.h"
+#include "SceneManagement.h"
+#include "ShaderParameterUtils.h"
+#include "UObject/ConstructorHelpers.h"
+
+// FScreenPassTextureViewportParameters and FScreenPassTextureInput
+#include "ScreenPass.h"
+#include "SceneTexturesConfig.h"
+
+
+
+BEGIN_GLOBAL_SHADER_PARAMETER_STRUCT(FCCRRegionDataInputParameter, )
+	SHADER_PARAMETER(FVector3f, Rotate)
+	SHADER_PARAMETER(FVector3f, Translate)
+	SHADER_PARAMETER(FVector3f, Scale)
+
+	SHADER_PARAMETER(float, WhiteTemp)
+	SHADER_PARAMETER(float, Tint)
+	SHADER_PARAMETER(float, Inner)
+	SHADER_PARAMETER(float, Outer)
+	SHADER_PARAMETER(float, Falloff)
+	SHADER_PARAMETER(float, Intensity)
+	SHADER_PARAMETER(float, FakeLight)
+	SHADER_PARAMETER(uint32, ExcludeStencil)
+	SHADER_PARAMETER(float, Invert)
+END_GLOBAL_SHADER_PARAMETER_STRUCT()
+
+
+BEGIN_GLOBAL_SHADER_PARAMETER_STRUCT(FCCRColorCorrectParameter, )
+	SHADER_PARAMETER(FVector4f, ColorSaturation)
+	SHADER_PARAMETER(FVector4f, ColorContrast)
+	SHADER_PARAMETER(FVector4f, ColorGamma)
+	SHADER_PARAMETER(FVector4f, ColorGain)
+	SHADER_PARAMETER(FVector4f, ColorOffset)
+END_GLOBAL_SHADER_PARAMETER_STRUCT()
+
+
+BEGIN_GLOBAL_SHADER_PARAMETER_STRUCT(FCCRColorCorrectShadowsParameter, )
+	SHADER_PARAMETER(FVector4f, ColorSaturation)
+	SHADER_PARAMETER(FVector4f, ColorContrast)
+	SHADER_PARAMETER(FVector4f, ColorGamma)
+	SHADER_PARAMETER(FVector4f, ColorGain)
+	SHADER_PARAMETER(FVector4f, ColorOffset)
+	SHADER_PARAMETER(float, ShadowMax)
+END_GLOBAL_SHADER_PARAMETER_STRUCT()
+
+
+BEGIN_GLOBAL_SHADER_PARAMETER_STRUCT(FCCRColorCorrectMidtonesParameter, )
+	SHADER_PARAMETER(FVector4f, ColorSaturation)
+	SHADER_PARAMETER(FVector4f, ColorContrast)
+	SHADER_PARAMETER(FVector4f, ColorGamma)
+	SHADER_PARAMETER(FVector4f, ColorGain)
+	SHADER_PARAMETER(FVector4f, ColorOffset)
+END_GLOBAL_SHADER_PARAMETER_STRUCT()
+
+
+BEGIN_GLOBAL_SHADER_PARAMETER_STRUCT(FCCRColorCorrectHighlightsParameter, )
+	SHADER_PARAMETER(FVector4f, ColorSaturation)
+	SHADER_PARAMETER(FVector4f, ColorContrast)
+	SHADER_PARAMETER(FVector4f, ColorGamma)
+	SHADER_PARAMETER(FVector4f, ColorGain)
+	SHADER_PARAMETER(FVector4f, ColorOffset)
+	SHADER_PARAMETER(float, HighlightsMin)
+END_GLOBAL_SHADER_PARAMETER_STRUCT()
+
+
+BEGIN_SHADER_PARAMETER_STRUCT(FCCRShaderInputParameters, )
+	SHADER_PARAMETER_RDG_TEXTURE(Texture2D<uint>, MergedStencilTexture)
+	SHADER_PARAMETER_STRUCT_REF(FViewUniformShaderParameters, View)
+	SHADER_PARAMETER_STRUCT_INCLUDE(FSceneTextureShaderParameters, SceneTextures)
+	SHADER_PARAMETER_STRUCT_REF(FWorkingColorSpaceShaderParameters, WorkingColorSpace)
+	SHADER_PARAMETER_STRUCT(FScreenPassTextureViewportParameters, PostProcessOutput)
+	SHADER_PARAMETER_STRUCT_ARRAY(FScreenPassTextureInput, PostProcessInput, [1])
+	RENDER_TARGET_BINDING_SLOTS()
+END_SHADER_PARAMETER_STRUCT()
+
+
+class FColorCorrectRegionsPostProcessMaterialShader : public FGlobalShader
+{
+public:
+	using FParameters = FCCRShaderInputParameters;
+	SHADER_USE_PARAMETER_STRUCT_WITH_LEGACY_BASE(FColorCorrectRegionsPostProcessMaterialShader, FGlobalShader);
+
+	static bool ShouldCompilePermutation(const FGlobalShaderPermutationParameters& Parameters)
+	{
+		return !IsMobilePlatform(Parameters.Platform);
+	}
+
+	static void ModifyCompilationEnvironment(const FGlobalShaderPermutationParameters& Parameters, FShaderCompilerEnvironment& OutEnvironment)
+	{
+		FMaterialShader::ModifyCompilationEnvironment(Parameters, OutEnvironment);
+	}
+
+};
+
+class FColorCorrectRegionMaterialVS : public FColorCorrectRegionsPostProcessMaterialShader
+{
+public:
+	DECLARE_GLOBAL_SHADER(FColorCorrectRegionMaterialVS);
+
+	FColorCorrectRegionMaterialVS() = default;
+	FColorCorrectRegionMaterialVS(const ShaderMetaType::CompiledShaderInitializerType & Initializer)
+		: FColorCorrectRegionsPostProcessMaterialShader(Initializer)
+	{}
+
+	void SetParameters(FRHIBatchedShaderParameters& BatchedParameters, const FSceneView& View)
+	{
+		FGlobalShader::SetParameters<FViewUniformShaderParameters>(BatchedParameters, View.ViewUniformBuffer);
+	}
+};
+
+class FColorCorrectGenericPS : public FColorCorrectRegionsPostProcessMaterialShader
+{
+public:
+	enum class ETemperatureType : uint8
+	{
+		LegacyTemperature = static_cast<uint8>(EColorCorrectRegionTemperatureType::LegacyTemperature),
+		WhiteBalance = static_cast<uint8>(EColorCorrectRegionTemperatureType::WhiteBalance),
+		ColorTemperature = static_cast<uint8>(EColorCorrectRegionTemperatureType::ColorTemperature),
+		Disabled,
+		MAX
+	};
+
+	class FTemperatureType : SHADER_PERMUTATION_ENUM_CLASS("TEMPERATURE_TYPE", ETemperatureType);
+	class FAdvancedShader : SHADER_PERMUTATION_BOOL("ADVANCED_CC");
+	class FStencilEnabled : SHADER_PERMUTATION_BOOL("STENCIL_ENABLED");
+
+	FColorCorrectGenericPS() = default;
+	FColorCorrectGenericPS(const ShaderMetaType::CompiledShaderInitializerType& Initializer)
+		: FColorCorrectRegionsPostProcessMaterialShader(Initializer)
+	{}
+
+	void SetParameters(FRHIBatchedShaderParameters& BatchedParameters, const FSceneView& View)
+	{
+		FGlobalShader::SetParameters<FViewUniformShaderParameters>(BatchedParameters, View.ViewUniformBuffer);
+	}
+};
+
+class FColorCorrectRegionMaterialPS : public FColorCorrectGenericPS
+{
+public:
+
+	DECLARE_GLOBAL_SHADER(FColorCorrectRegionMaterialPS);
+
+	class FShaderType : SHADER_PERMUTATION_ENUM_CLASS("SHAPE_TYPE", EColorCorrectRegionsType);
+
+	using FPermutationDomain = TShaderPermutationDomain<FShaderType, FTemperatureType, FAdvancedShader, FStencilEnabled>;
+
+	FColorCorrectRegionMaterialPS() = default;
+	FColorCorrectRegionMaterialPS(const ShaderMetaType::CompiledShaderInitializerType & Initializer)
+		: FColorCorrectGenericPS(Initializer)
+	{}
+};
+
+
+class FColorCorrectWindowMaterialPS : public FColorCorrectGenericPS
+{
+public:
+
+	DECLARE_GLOBAL_SHADER(FColorCorrectWindowMaterialPS);
+
+	class FShaderType : SHADER_PERMUTATION_ENUM_CLASS("WINDOW_TYPE", EColorCorrectWindowType);
+
+	using FPermutationDomain = TShaderPermutationDomain<FShaderType, FTemperatureType, FAdvancedShader, FStencilEnabled>;
+
+	FColorCorrectWindowMaterialPS() = default;
+	FColorCorrectWindowMaterialPS(const ShaderMetaType::CompiledShaderInitializerType& Initializer)
+		: FColorCorrectGenericPS(Initializer)
+	{}
+};
+
+// The vertex shader used by DrawScreenPass to draw a rectangle.
+class FColorCorrectScreenPassVS : public FGlobalShader
+{
+public:
+	DECLARE_GLOBAL_SHADER(FColorCorrectScreenPassVS);
+
+	FColorCorrectScreenPassVS() = default;
+	FColorCorrectScreenPassVS(const ShaderMetaType::CompiledShaderInitializerType& Initializer)
+		: FGlobalShader(Initializer)
+	{}
+};
+
+// A simple shader that outputs (0.,0.,0.,0.)
+class FClearRectPS : public FGlobalShader
+{
+	DECLARE_GLOBAL_SHADER(FClearRectPS);
+	SHADER_USE_PARAMETER_STRUCT(FClearRectPS, FGlobalShader);
+
+	static void ModifyCompilationEnvironment(const FGlobalShaderPermutationParameters& Parameters, FShaderCompilerEnvironment& OutEnvironment)
+	{
+		FColorCorrectRegionMaterialPS::ModifyCompilationEnvironment(Parameters, OutEnvironment);
+		OutEnvironment.SetDefine(TEXT("CLEAR_PS"), 1);
+	}
+
+	BEGIN_SHADER_PARAMETER_STRUCT(FParameters, )
+		RENDER_TARGET_BINDING_SLOTS()
+	END_SHADER_PARAMETER_STRUCT()
+};
+
+
+/** Vertex shader to rasterize obstructive CCR Plane. */
+class FCCRStencilMergerVS : public FGlobalShader
+{
+public:
+	DECLARE_GLOBAL_SHADER(FCCRStencilMergerVS);
+
+	FCCRStencilMergerVS() = default;
+	FCCRStencilMergerVS(const ShaderMetaType::CompiledShaderInitializerType& Initializer)
+		: FGlobalShader(Initializer)
+	{}
+};
+
+/** Pixel shader to rasterize obstructive CCR Plane. */
+class FCCRStencilMergerPS : public FGlobalShader
+{
+	DECLARE_GLOBAL_SHADER(FCCRStencilMergerPS);
+	SHADER_USE_PARAMETER_STRUCT(FCCRStencilMergerPS, FGlobalShader);
+
+	static bool ShouldCompilePermutation(const FGlobalShaderPermutationParameters& Parameters)
+	{
+		return !IsMobilePlatform(Parameters.Platform);
+	}
+
+	BEGIN_SHADER_PARAMETER_STRUCT(FParameters, )
+		SHADER_PARAMETER_STRUCT_REF(FViewUniformShaderParameters, View)
+		SHADER_PARAMETER_RDG_BUFFER_SRV(StructuredBuffer<uint>, StencilIds)
+		SHADER_PARAMETER(uint32, StencilIdCount)
+		SHADER_PARAMETER_STRUCT_INCLUDE(FSceneTextureShaderParameters, SceneTextures)
+		SHADER_PARAMETER_STRUCT(FScreenPassTextureViewportParameters, PostProcessOutput)
+		RENDER_TARGET_BINDING_SLOTS()
+	END_SHADER_PARAMETER_STRUCT()
+};
diff --git a/Engine/Plugins/Experimental/ColorCorrectRegions/Source/ColorCorrectRegions/Private/ColorCorrectRegionsSceneViewExtension.h b/Engine/Plugins/Experimental/ColorCorrectRegions/Source/ColorCorrectRegions/Private/ColorCorrectRegionsSceneViewExtension.h
new file mode 100644
--- /dev/null
+++ b/Engine/Plugins/Experimental/ColorCorrectRegions/Source/ColorCorrectRegions/Private/ColorCorrectRegionsSceneViewExtension.h
@@ -0,0 +1,28 @@
+// Copyright Epic Games, Inc. All Rights Reserved.
+#pragma once
+#include "SceneViewExtension.h"
+#include "RHI.h"
+#include "RHIResources.h"
+
+class UColorCorrectRegionsSubsystem;
+class UMaterialInterface;
+class FRDGTexture;
+
+class FColorCorrectRegionsSceneViewExtension : public FWorldSceneViewExtension
+{
+public:
+	FColorCorrectRegionsSceneViewExtension(const FAutoRegister& AutoRegister, UWorld* InWorld, UColorCorrectRegionsSubsystem* InWorldSubsystem);
+	
+	//~ Begin FSceneViewExtensionBase Interface
+	virtual void SetupViewFamily(FSceneViewFamily& InViewFamily) override {};
+	virtual void SetupView(FSceneViewFamily& InViewFamily, FSceneView& InView) override {};
+	virtual void BeginRenderViewFamily(FSceneViewFamily& InViewFamily) override;
+	virtual void PrePostProcessPass_RenderThread(FRDGBuilder& GraphBuilder, const FSceneView& View, const FPostProcessingInputs& Inputs) override;
+	//~ End FSceneViewExtensionBase Interface
+public:
+	// Called when owning subsystem needs to release this extension.
+	void Invalidate();
+private:
+	UColorCorrectRegionsSubsystem* WorldSubsystem;
+};
+	
\ No newline at end of file
diff --git a/Engine/Plugins/Experimental/ColorCorrectRegions/Source/ColorCorrectRegions/Private/ColorCorrectRegionsStencilManager.h b/Engine/Plugins/Experimental/ColorCorrectRegions/Source/ColorCorrectRegions/Private/ColorCorrectRegionsStencilManager.h
new file mode 100644
--- /dev/null
+++ b/Engine/Plugins/Experimental/ColorCorrectRegions/Source/ColorCorrectRegions/Private/ColorCorrectRegionsStencilManager.h
@@ -0,0 +1,22 @@
+// Copyright Epic Games, Inc. All Rights Reserved.
+
+#pragma once
+
+#include "CoreMinimal.h"
+#include "ColorCorrectRegion.h"
+#include "GameFramework/Actor.h"
+
+
+class FColorCorrectRegionsStencilManager
+{
+public:
+		static void AssignStencilNumberToActorForSelectedRegion(UWorld* CurrentWorld, AColorCorrectRegion* Region, TSoftObjectPtr<AActor> ActorToAssignStencilTo, bool bIgnoreUserNotifications, bool bSoftAssign);
+		static void RemoveStencilNumberForSelectedRegion(UWorld* CurrentWorld, AColorCorrectRegion* Region);
+		static void OnCCRRemoved(UWorld* CurrentWorld, AColorCorrectRegion* Region);
+		static void ClearInvalidActorsForSelectedRegion(AColorCorrectRegion* Region);
+		static void AssignStencilIdsToAllActorsForCCR(UWorld* CurrentWorld, AColorCorrectRegion* Region, bool bIgnoreUserNotifications, bool bSoftAssign);
+		static void CleanActor(AActor* Actor);
+
+		/** Runs over all actors assigned to the selected Region and makes sure that Stencil Ids match the ones assigned to it. */
+		static void CheckAssignedActorsValidity(AColorCorrectRegion* Region);
+};
\ No newline at end of file
diff --git a/Engine/Plugins/Experimental/ColorCorrectRegions/Source/ColorCorrectRegions/Private/ColorCorrectRegionsSubsystem.h b/Engine/Plugins/Experimental/ColorCorrectRegions/Source/ColorCorrectRegions/Private/ColorCorrectRegionsSubsystem.h
new file mode 100644
--- /dev/null
+++ b/Engine/Plugins/Experimental/ColorCorrectRegions/Source/ColorCorrectRegions/Private/ColorCorrectRegionsSubsystem.h
@@ -0,0 +1,106 @@
+// Copyright Epic Games, Inc. All Rights Reserved.
+
+#pragma once
+#include "CoreMinimal.h"
+#include "UObject/ObjectMacros.h"
+#include "UObject/UObjectGlobals.h"
+#include "UObject/Object.h"
+#include "Subsystems/EngineSubsystem.h"
+#include "Engine/EngineBaseTypes.h"
+#include "ColorCorrectRegion.h"
+#include "ColorCorrectRegionsSceneViewExtension.h"
+
+#if WITH_EDITOR
+#include "EditorUndoClient.h"
+#endif
+
+#include "ColorCorrectRegionsSubsystem.generated.h"
+
+/**
+ * World Subsystem responsible for managing AColorCorrectRegion classes in level.
+ * This subsystem handles:
+ *		Level Loaded, Undo/Redo, Added to level, Removed from level events.
+ * Unfortunately AActor class itself is not aware of when it is added/removed, Undo/Redo etc in the level.
+ * 
+ * This is the only way (that we found) that was handling all region aggregation cases in more or less efficient way.
+ *		Covered cases: Region added to a level, deleted from level, level loaded, undo, redo, level closed, editor closed:
+ *		World subsystem keeps track of all Regions in a level via three events OnLevelActorAdded and OnLevelActorDeleted.
+ *		Actor classes are unaware of when they are added/deleted/undo/redo etc in the level, therefore this is the best place to manage this.
+ * Alternative strategies (All tested):
+ *		World's AddOnActorSpawnedHandler. Flawed. Invoked in some cases we don't need, but does not get called during UNDO/REDO
+ *		AActor's PostSpawnInitialize, PostActorCreated  and OnConstruction are also flawed.
+ *		AActor does not have an internal event for when its deleted (EndPlay is the closest we have).
+ */
+UCLASS()
+class UColorCorrectRegionsSubsystem : public UTickableWorldSubsystem
+{
+	GENERATED_BODY()
+public:
+
+	// Subsystem Init/Deinit
+	virtual void Initialize(FSubsystemCollectionBase& Collection) override;
+	virtual void Deinitialize() override;
+
+	virtual bool IsTickableInEditor() const { return true; }
+	virtual void Tick(float DeltaTime) override;
+	virtual TStatId GetStatId() const override { RETURN_QUICK_DECLARE_CYCLE_STAT(UColorCorrectRegionsSubsystem, STATGROUP_Tickables); }
+
+public:
+
+	/** A callback for CC Region creation. */
+	void OnActorSpawned(AActor* InActor);
+
+	/** A callback for CC Region deletion. */
+	void OnActorDeleted(AActor* InActor, bool bClearStencilIdValues);
+
+	/** Called when duplication process is started in the level. */
+	void OnDuplicateActorsBegin() { bDuplicationStarted = true; };
+
+	/** Called when duplication process is ended in the level. */
+	void OnDuplicateActorsEnd();
+
+	/** Handles Stencil Ids for the selected CCR and corresponding actor. */
+	void AssignStencilIdsToPerActorCC(AColorCorrectRegion* Region, bool bIgnoreUserNotificaion = false, bool bSoftAssign = false);
+	
+	/** Handles removal of Stencil Ids for the selected CCR. */
+	void ClearStencilIdsToPerActorCC(AColorCorrectRegion* Region);	
+	
+	/** Handles cases when stencil Id has been changed from outside by the user manually. */
+	void CheckAssignedActorsValidity(AColorCorrectRegion* Region);
+
+private:
+
+	/** Repopulates array of region actors. */
+	void RefreshRegions();
+
+	/**
+	* Copy states required for rendering to be consumed by Scene view extension to render all active 
+	* CCRs and CCWs.
+	*/
+	void TransferStates();
+
+private:
+
+	/** Stores pointers to ColorCorrectRegion Actors that use priority for sorting. */
+	TArray<TWeakObjectPtr<AColorCorrectRegion>> RegionsPriorityBased;
+	/** Stores pointers to ColorCorrectRegion Actors that are based on distance from camera. */
+	TArray<TWeakObjectPtr<AColorCorrectRegion>> RegionsDistanceBased;
+
+
+	/** Proxies to be used exclusively on render thread. Copies of the state of CC Actors sorted by priority. */
+	TArray<FColorCorrectRenderProxyPtr> ProxiesPriorityBased;
+	/** Proxies to be used exclusively on render thread. Copies of the state of CC Actors sorted by distance. */
+	TArray<FColorCorrectRenderProxyPtr> ProxiesDistanceBased;
+
+	TSharedPtr< class FColorCorrectRegionsSceneViewExtension, ESPMode::ThreadSafe > PostProcessSceneViewExtension;
+
+	/** This is to handle actor duplication for Per Actor CC. */
+	bool bDuplicationStarted = false;
+	TArray<AActor*> DuplicatedActors;
+
+	// This is for optimization purposes that would let us check assigned actors component's stencil ids every once in a while.
+	float TimeSinceLastValidityCheck = 0;
+
+public:
+	friend class FColorCorrectRegionsSceneViewExtension;
+};
\ No newline at end of file
diff --git a/Engine/Plugins/Experimental/ColorCorrectRegions/Source/ColorCorrectRegions/Public/ColorCorrectRegionsPostProcessMaterial.h b/Engine/Plugins/Experimental/ColorCorrectRegions/Source/ColorCorrectRegions/Public/ColorCorrectRegionsPostProcessMaterial.h
deleted file mode 100644
diff --git a/Engine/Plugins/Experimental/ColorCorrectRegions/Source/ColorCorrectRegions/Public/ColorCorrectRegionsSceneViewExtension.h b/Engine/Plugins/Experimental/ColorCorrectRegions/Source/ColorCorrectRegions/Public/ColorCorrectRegionsSceneViewExtension.h
deleted file mode 100644
diff --git a/Engine/Plugins/Experimental/ColorCorrectRegions/Source/ColorCorrectRegions/Public/ColorCorrectRegionsStencilManager.h b/Engine/Plugins/Experimental/ColorCorrectRegions/Source/ColorCorrectRegions/Public/ColorCorrectRegionsStencilManager.h
deleted file mode 100644
diff --git a/Engine/Plugins/Experimental/ColorCorrectRegions/Source/ColorCorrectRegions/Public/ColorCorrectRegionsSubsystem.h b/Engine/Plugins/Experimental/ColorCorrectRegions/Source/ColorCorrectRegions/Public/ColorCorrectRegionsSubsystem.h
deleted file mode 100644
