Color Correction Actors Fix: Call Modify() on PrimitiveComponents directly instead of on the Actor.

[RN] Color Correction Actors in Multi-User: assigning some actors with mesh components to per actor color correction was not transacting over multi-user.
#Synced-CL 0

diff --git a/Engine/Plugins/Experimental/ColorCorrectRegions/Source/ColorCorrectRegions/Private/ColorCorrectRegionsStencilManager.cpp b/Engine/Plugins/Experimental/ColorCorrectRegions/Source/ColorCorrectRegions/Private/ColorCorrectRegionsStencilManager.cpp
--- a/Engine/Plugins/Experimental/ColorCorrectRegions/Source/ColorCorrectRegions/Private/ColorCorrectRegionsStencilManager.cpp
+++ b/Engine/Plugins/Experimental/ColorCorrectRegions/Source/ColorCorrectRegions/Private/ColorCorrectRegionsStencilManager.cpp
@@ -111,12 +111,12 @@ namespace
 			}
 			if (!bSoftAssign)
 			{
+#if WITH_EDITOR
+				PrimitiveComponent->Modify();
+#endif
 				PrimitiveComponent->SetRenderCustomDepth(true);
 				PrimitiveComponent->CustomDepthStencilValue = static_cast<int32>(StencilData->AssignedStencil);
 				PrimitiveComponent->MarkRenderStateDirty();
-#if WITH_EDITOR
-				ActorToAssignStencilTo->Modify();
-#endif
 			}
 			else
 			{
@@ -163,12 +163,12 @@ namespace
 
 				for (UPrimitiveComponent* PrimitiveComponent : PrimitiveComponentsToDirty)
 				{
-					PrimitiveComponent->SetRenderCustomDepth(false);
-					PrimitiveComponent->CustomDepthStencilValue = 0;
-					PrimitiveComponent->MarkRenderStateDirty();
 #if WITH_EDITOR
 					PrimitiveComponent->Modify();
 #endif
+					PrimitiveComponent->SetRenderCustomDepth(false);
+					PrimitiveComponent->CustomDepthStencilValue = 0;
+					PrimitiveComponent->MarkRenderStateDirty();
 				}
 			}
 		}
