Color Correction Actors Multi-User: fixed an issue where stencil id's assignment on some actors were not transacted over multi-user.
- Replace polling-based Concert sync request with a reliable custom event (FConcertCCRPerActorAssignmentEvent);
- Fix undo/redo: stencil changes now share a transaction with AffectedActors, and PrimitiveComponent->Modify() is called before modifying stencil values

[RN] Fixed an issue with CCRs not assigning stencil ids when transacted over multi-user

#Synced-CL 0

diff --git a/Engine/Plugins/Experimental/ColorCorrectRegions/ColorCorrectRegions.uplugin b/Engine/Plugins/Experimental/ColorCorrectRegions/ColorCorrectRegions.uplugin
--- a/Engine/Plugins/Experimental/ColorCorrectRegions/ColorCorrectRegions.uplugin
+++ b/Engine/Plugins/Experimental/ColorCorrectRegions/ColorCorrectRegions.uplugin
@@ -38,6 +38,11 @@
 		{
 			"Name": "ObjectMixer",
 			"Enabled": true
+		},
+		{
+			"Name": "ConcertSyncClient",
+			"Enabled": false,
+			"Optional": true
 		}
 	]
 }
\ No newline at end of file
diff --git a/Engine/Plugins/Experimental/ColorCorrectRegions/Source/ColorCorrectRegions/ColorCorrectRegions.Build.cs b/Engine/Plugins/Experimental/ColorCorrectRegions/Source/ColorCorrectRegions/ColorCorrectRegions.Build.cs
--- a/Engine/Plugins/Experimental/ColorCorrectRegions/Source/ColorCorrectRegions/ColorCorrectRegions.Build.cs
+++ b/Engine/Plugins/Experimental/ColorCorrectRegions/Source/ColorCorrectRegions/ColorCorrectRegions.Build.cs
@@ -33,6 +33,15 @@ public class ColorCorrectRegions : ModuleRules
 			PrivateDependencyModuleNames.Add("EditorWidgets");
 			PrivateDependencyModuleNames.Add("Slate");
 			PrivateDependencyModuleNames.Add("SlateCore");
+
+			// Multi-user support: synchronise per actor CC stencil assignments via Concert
+			// custom events so passive clients don't rely on tick-based polling.
+			PrivateDependencyModuleNames.AddRange(
+				new string[]
+				{
+					"ConcertSyncClient",
+				}
+			);
 		}
 	}
 }
diff --git a/Engine/Plugins/Experimental/ColorCorrectRegions/Source/ColorCorrectRegions/Private/ColorCorrectRegion.cpp b/Engine/Plugins/Experimental/ColorCorrectRegions/Source/ColorCorrectRegions/Private/ColorCorrectRegion.cpp
--- a/Engine/Plugins/Experimental/ColorCorrectRegions/Source/ColorCorrectRegions/Private/ColorCorrectRegion.cpp
+++ b/Engine/Plugins/Experimental/ColorCorrectRegions/Source/ColorCorrectRegions/Private/ColorCorrectRegion.cpp
@@ -1,7 +1,6 @@
 // Copyright Epic Games, Inc. All Rights Reserved.
 
 #include "ColorCorrectRegion.h"
-#include "Async/Async.h"
 #include "ColorCorrectRegionsModule.h"
 #include "ColorCorrectRegionsSubsystem.h"
 #include "ColorCorrectWindow.h"
@@ -18,6 +17,7 @@
 
 #if WITH_EDITOR
 #include "IDisplayClusterLightCardExtenderModule.h"
+#include "Misc/TransactionObjectEvent.h"
 #endif
 
 ENUM_RANGE_BY_COUNT(EColorCorrectRegionsType, EColorCorrectRegionsType::MAX)
@@ -82,8 +82,20 @@ void AColorCorrectRegion::HandleAffectedActorsPropertyChange(uint32 ActorListCha
 		ColorCorrectRegionsSubsystem = World->GetSubsystem<UColorCorrectRegionsSubsystem>();
 	}
 
+	if (ActorListChangeType == EPropertyChangeType::ArrayClear
+		|| ActorListChangeType == EPropertyChangeType::ArrayRemove
+		|| ActorListChangeType == EPropertyChangeType::ValueSet
+		|| ActorListChangeType == EPropertyChangeType::Unspecified)
+	{
+		if (ColorCorrectRegionsSubsystem.IsValid())
+		{
+			ColorCorrectRegionsSubsystem->ClearStencilIdsToPerActorCC(this);
+		}
+	}
+
 	if (ActorListChangeType == EPropertyChangeType::ArrayAdd
-		|| ActorListChangeType == EPropertyChangeType::ValueSet)
+		|| ActorListChangeType == EPropertyChangeType::ValueSet
+		|| ActorListChangeType == EPropertyChangeType::Unspecified)
 	{
 		// In case user assigns Color Correct Region or Window, we should remove it as it is invalid operation.
 		{
@@ -112,14 +124,10 @@ void AColorCorrectRegion::HandleAffectedActorsPropertyChange(uint32 ActorListCha
 		}
 	}
 
-	if (ActorListChangeType == EPropertyChangeType::ArrayClear
-		|| ActorListChangeType == EPropertyChangeType::ArrayRemove
-		|| ActorListChangeType == EPropertyChangeType::ValueSet)
+	// Notify connected multi-user clients of the updated AffectedActors list via a Concert custom event
+	if (ColorCorrectRegionsSubsystem.IsValid())
 	{
-		if (ColorCorrectRegionsSubsystem.IsValid())
-		{
-			ColorCorrectRegionsSubsystem->ClearStencilIdsToPerActorCC(this);
-		}
+		ColorCorrectRegionsSubsystem->NotifyPerActorCCChanged(this);
 	}
 }
 
@@ -183,13 +191,16 @@ void AColorCorrectRegion::PostEditChangeProperty(struct FPropertyChangedEvent& P
 {
 	const FName PropertyName = PropertyChangedEvent.GetPropertyName();
 
-	if (PropertyName == GET_MEMBER_NAME_CHECKED(AColorCorrectRegion, AffectedActors))
+	// Use MemberProperty (the class-level field) rather than the leaf property name.
+	// When resetting an array element to default, the leaf property is the inner element type,
+	// not "AffectedActors" itself, so GetPropertyName() would miss it.
+	const FName MemberPropertyName = PropertyChangedEvent.MemberProperty
+		? PropertyChangedEvent.MemberProperty->GetFName()
+		: PropertyName;
+
+	if (MemberPropertyName == GET_MEMBER_NAME_CHECKED(AColorCorrectRegion, AffectedActors))
 	{
-		/** Since there might be Dialogs involved we need to run this on game thread. */
-		AsyncTask(ENamedThreads::GameThread, [this, ActorListChangeType = PropertyChangedEvent.ChangeType]() 
-		{
-			HandleAffectedActorsPropertyChange(ActorListChangeType);
-		});
+		HandleAffectedActorsPropertyChange(PropertyChangedEvent.ChangeType);
 	}
 
 	// Stage actor properties
@@ -232,6 +243,14 @@ void AColorCorrectRegion::PostTransacted(const FTransactionObjectEvent& Transact
 {
 	Super::PostTransacted(TransactionEvent);
 	FixMeshComponentReferences();
+
+	// PerAffectedActorStencilData is non-UPROPERTY so the transaction system never restores it.
+	// Clear it on undo/redo so it stays consistent with the transaction-restored AffectedActors.
+	// The next stencil assignment will repopulate it.
+	if (TransactionEvent.GetEventType() == ETransactionObjectEventType::UndoRedo)
+	{
+		PerAffectedActorStencilData.Empty();
+	}
 }
 #endif //WITH_EDITOR
 
diff --git a/Engine/Plugins/Experimental/ColorCorrectRegions/Source/ColorCorrectRegions/Private/ColorCorrectRegionSubsystem.cpp b/Engine/Plugins/Experimental/ColorCorrectRegions/Source/ColorCorrectRegions/Private/ColorCorrectRegionSubsystem.cpp
--- a/Engine/Plugins/Experimental/ColorCorrectRegions/Source/ColorCorrectRegions/Private/ColorCorrectRegionSubsystem.cpp
+++ b/Engine/Plugins/Experimental/ColorCorrectRegions/Source/ColorCorrectRegions/Private/ColorCorrectRegionSubsystem.cpp
@@ -169,6 +169,14 @@ void UColorCorrectRegionsSubsystem::Initialize(FSubsystemCollectionBase& Collect
 
 	// Initializing Scene view extension responsible for rendering regions.
 	PostProcessSceneViewExtension = FSceneViewExtensions::NewExtension<FColorCorrectRegionsSceneViewExtension>(GetWorld(), this);
+
+#if WITH_EDITOR
+	if (GetWorld()->WorldType == EWorldType::Editor)
+	{
+		MultiUserManager = MakeUnique<FCCRMultiUserManager>(this);
+	}
+#endif
+
 	Super::Initialize(Collection);
 }
 
@@ -177,6 +185,8 @@ void UColorCorrectRegionsSubsystem::Deinitialize()
 #if WITH_EDITOR
 	if (GetWorld()->WorldType == EWorldType::Editor)
 	{
+		MultiUserManager.Reset();
+
 		GEngine->OnLevelActorAdded().RemoveAll(this);
 		GEngine->OnLevelActorDeleted().RemoveAll(this);
 
@@ -242,19 +252,22 @@ void UColorCorrectRegionsSubsystem::Tick(float DeltaTime)
 		const float WaitTimeInSecs = 1.0;
 		if (TimeSinceLastValidityCheck >= WaitTimeInSecs)
 		{
-			for (TWeakObjectPtr<AColorCorrectRegion> Region : RegionsPriorityBased)
+			auto CheckRegion = [this](const TWeakObjectPtr<AColorCorrectRegion>& Region)
 			{
-				if (Region.IsValid())
+				if (!Region.IsValid())
 				{
-					CheckAssignedActorsValidity(Region.Get());
+					return;
 				}
+				CheckAssignedActorsValidity(Region.Get());
+			};
+
+			for (const TWeakObjectPtr<AColorCorrectRegion>& Region : RegionsPriorityBased)
+			{
+				CheckRegion(Region);
 			}
-			for (TWeakObjectPtr<AColorCorrectRegion> Region : RegionsDistanceBased)
+			for (const TWeakObjectPtr<AColorCorrectRegion>& Region : RegionsDistanceBased)
 			{
-				if (Region.IsValid())
-				{
-					CheckAssignedActorsValidity(Region.Get());
-				}
+				CheckRegion(Region);
 			}
 
 			TimeSinceLastValidityCheck = 0;
@@ -419,6 +432,16 @@ void UColorCorrectRegionsSubsystem::CheckAssignedActorsValidity(AColorCorrectReg
 	FColorCorrectRegionsStencilManager::CheckAssignedActorsValidity(Region);
 }
 
+void UColorCorrectRegionsSubsystem::NotifyPerActorCCChanged(AColorCorrectRegion* Region)
+{
+#if WITH_EDITOR
+	if (MultiUserManager)
+	{
+		MultiUserManager->SendPerActorAssignmentEvent(Region);
+	}
+#endif
+}
+
 void UColorCorrectRegionsSubsystem::RefreshRegions()
 {
 	RegionsPriorityBased.Reset();
diff --git a/Engine/Plugins/Experimental/ColorCorrectRegions/Source/ColorCorrectRegions/Private/ColorCorrectRegionsSubsystem.h b/Engine/Plugins/Experimental/ColorCorrectRegions/Source/ColorCorrectRegions/Private/ColorCorrectRegionsSubsystem.h
--- a/Engine/Plugins/Experimental/ColorCorrectRegions/Source/ColorCorrectRegions/Private/ColorCorrectRegionsSubsystem.h
+++ b/Engine/Plugins/Experimental/ColorCorrectRegions/Source/ColorCorrectRegions/Private/ColorCorrectRegionsSubsystem.h
@@ -14,6 +14,10 @@
 #include "EditorUndoClient.h"
 #endif
 
+#if WITH_EDITOR
+#include "MultiUser/ColorCorrectRegionsMultiUserManager.h"
+#endif
+
 #include "ColorCorrectRegionsSubsystem.generated.h"
 
 /**
@@ -68,6 +72,13 @@ public:
 	/** Handles cases when stencil Id has been changed from outside by the user manually. */
 	void CheckAssignedActorsValidity(AColorCorrectRegion* Region);
 
+	/**
+	 * Notifies connected multi-user clients that per-actor CC stencil assignment changed for
+	 * the given region. Called after local stencil operations complete. No-op when no Concert
+	 * session is active or in non-editor builds.
+	 */
+	void NotifyPerActorCCChanged(AColorCorrectRegion* Region);
+
 private:
 
 	/** Repopulates array of region actors. */
@@ -101,6 +112,11 @@ private:
 	// This is for optimization purposes that would let us check assigned actors component's stencil ids every once in a while.
 	float TimeSinceLastValidityCheck = 0;
 
+#if WITH_EDITOR
+	/** Sends Concert custom events when Per Actor CC assignment changes. Null when Concert is unavailable. */
+	TUniquePtr<FCCRMultiUserManager> MultiUserManager;
+#endif
+
 public:
 	friend class FColorCorrectRegionsSceneViewExtension;
 };
\ No newline at end of file
diff --git a/Engine/Plugins/Experimental/ColorCorrectRegions/Source/ColorCorrectRegions/Private/MultiUser/ColorCorrectRegionsMultiUserManager.cpp b/Engine/Plugins/Experimental/ColorCorrectRegions/Source/ColorCorrectRegions/Private/MultiUser/ColorCorrectRegionsMultiUserManager.cpp
new file mode 100644
--- /dev/null
+++ b/Engine/Plugins/Experimental/ColorCorrectRegions/Source/ColorCorrectRegions/Private/MultiUser/ColorCorrectRegionsMultiUserManager.cpp
@@ -0,0 +1,167 @@
+// Copyright Epic Games, Inc. All Rights Reserved.
+
+#if WITH_EDITOR
+
+#include "MultiUser/ColorCorrectRegionsMultiUserManager.h"
+#include "MultiUser/ColorCorrectRegionsSyncData.h"
+#include "ColorCorrectRegion.h"
+#include "ColorCorrectRegionsSubsystem.h"
+
+#include "ConcertTransportMessages.h"
+#include "IConcertClient.h"
+#include "IConcertSession.h"
+#include "IConcertSessionHandler.h"
+#include "IConcertSyncClient.h"
+#include "IConcertSyncClientModule.h"
+
+#include "Engine/World.h"
+#include "EngineUtils.h"
+
+
+FCCRMultiUserManager::FCCRMultiUserManager(UColorCorrectRegionsSubsystem* InSubsystem)
+	: WeakSubsystem(InSubsystem)
+{
+	IConcertSyncClientModule* ConcertModule = FModuleManager::GetModulePtr<IConcertSyncClientModule>("ConcertSyncClient");
+	if (!ConcertModule)
+	{
+		return;
+	}
+
+	TSharedPtr<IConcertSyncClient> ConcertSyncClient = ConcertModule->GetClient(TEXT("MultiUser"));
+	if (!ConcertSyncClient)
+	{
+		return;
+	}
+
+	IConcertClientRef Client = ConcertSyncClient->GetConcertClient();
+	ConcertClient = Client;
+	SessionStartupHandle = Client->OnSessionStartup().AddRaw(this, &FCCRMultiUserManager::OnSessionStartup);
+	SessionShutdownHandle = Client->OnSessionShutdown().AddRaw(this, &FCCRMultiUserManager::OnSessionShutdown);
+	SessionConnectionChangedHandle = Client->OnSessionConnectionChanged().AddRaw(this, &FCCRMultiUserManager::OnSessionConnectionChanged);
+
+	// If a session is already active when the subsystem initialises (e.g. level loaded while
+	// already in multi-user), register the event handler immediately.
+	if (TSharedPtr<IConcertClientSession> CurrentSession = Client->GetCurrentSession())
+	{
+		OnSessionStartup(CurrentSession.ToSharedRef());
+	}
+}
+
+FCCRMultiUserManager::~FCCRMultiUserManager()
+{
+	if (TSharedPtr<IConcertClient> Client = ConcertClient.Pin())
+	{
+		Client->OnSessionStartup().Remove(SessionStartupHandle);
+		Client->OnSessionShutdown().Remove(SessionShutdownHandle);
+		Client->OnSessionConnectionChanged().Remove(SessionConnectionChangedHandle);
+	}
+
+	if (TSharedPtr<IConcertClientSession> Session = WeakSession.Pin())
+	{
+		Session->UnregisterCustomEventHandler<FConcertCCRPerActorAssignmentEvent>(this);
+	}
+}
+
+void FCCRMultiUserManager::OnSessionStartup(TSharedRef<IConcertClientSession> Session)
+{
+	if (TSharedPtr<IConcertClientSession> PreviousSession = WeakSession.Pin())
+	{
+		PreviousSession->UnregisterCustomEventHandler<FConcertCCRPerActorAssignmentEvent>(this);
+	}
+	WeakSession = Session;
+	Session->RegisterCustomEventHandler<FConcertCCRPerActorAssignmentEvent>(
+		this, &FCCRMultiUserManager::HandlePerActorAssignmentEvent);
+}
+
+void FCCRMultiUserManager::OnSessionShutdown(TSharedRef<IConcertClientSession> Session)
+{
+	Session->UnregisterCustomEventHandler<FConcertCCRPerActorAssignmentEvent>(this);
+	WeakSession.Reset();
+}
+
+void FCCRMultiUserManager::OnSessionConnectionChanged(IConcertClientSession& Session, EConcertConnectionStatus ConnectionStatus)
+{
+	if (ConnectionStatus == EConcertConnectionStatus::Connected)
+	{
+		if (TSharedPtr<IConcertClient> Client = ConcertClient.Pin())
+		{
+			if (TSharedPtr<IConcertClientSession> CurrentSession = Client->GetCurrentSession())
+			{
+				// Delegate to OnSessionStartup so handler registration and WeakSession
+				// are set from the same TSharedRef, avoiding a mismatch with Session&.
+				OnSessionStartup(CurrentSession.ToSharedRef());
+			}
+		}
+	}
+	else if (ConnectionStatus == EConcertConnectionStatus::Disconnecting)
+	{
+		Session.UnregisterCustomEventHandler<FConcertCCRPerActorAssignmentEvent>(this);
+		WeakSession.Reset();
+	}
+}
+
+void FCCRMultiUserManager::SendPerActorAssignmentEvent(AColorCorrectRegion* Region)
+{
+	TSharedPtr<IConcertClientSession> Session = WeakSession.Pin();
+	if (!Session || !Region)
+	{
+		return;
+	}
+
+	FConcertCCRPerActorAssignmentEvent Event;
+	Event.CCRActorPath = FSoftObjectPath(Region);
+
+	// Snapshot the full AffectedActors list so the receiver doesn't depend on Concert's
+	// property transaction having arrived first.
+	for (const TSoftObjectPtr<AActor>& Actor : Region->AffectedActors)
+	{
+		Event.AffectedActors.Add(Actor.ToSoftObjectPath());
+	}
+
+	Session->SendCustomEvent(Event, Session->GetSessionClientEndpointIds(), EConcertMessageFlags::ReliableOrdered);
+}
+
+void FCCRMultiUserManager::HandlePerActorAssignmentEvent(
+	const FConcertSessionContext& Context,
+	const FConcertCCRPerActorAssignmentEvent& Event)
+{
+	// Ignore events that originated from this machine, or if the session is no longer valid.
+	TSharedPtr<IConcertClientSession> Session = WeakSession.Pin();
+	if (!Session || Context.SourceEndpointId == Session->GetSessionClientEndpointId())
+	{
+		return;
+	}
+
+	UColorCorrectRegionsSubsystem* Subsystem = WeakSubsystem.Get();
+	if (!Subsystem)
+	{
+		return;
+	}
+
+	UWorld* World = Subsystem->GetWorld();
+	if (!World)
+	{
+		return;
+	}
+
+	AColorCorrectRegion* Region = Cast<AColorCorrectRegion>(Event.CCRActorPath.ResolveObject());
+	if (!Region || Region->GetWorld() != World)
+	{
+		return;
+	}
+
+	// Populate AffectedActors from the event so stencil assignment has the correct actor list
+	// regardless of whether Concert's own property transaction has arrived yet.
+	Region->AffectedActors.Empty();
+	for (const FSoftObjectPath& Path : Event.AffectedActors)
+	{
+		Region->AffectedActors.Add(TSoftObjectPtr<AActor>(Path));
+	}
+
+	// Clear all stencil IDs for the region, then re-assign for the current AffectedActors list.
+	Subsystem->ClearStencilIdsToPerActorCC(Region);
+	Subsystem->AssignStencilIdsToPerActorCC(Region, /*bIgnoreUserNotification=*/true);
+}
+
+
+#endif // WITH_EDITOR
diff --git a/Engine/Plugins/Experimental/ColorCorrectRegions/Source/ColorCorrectRegions/Private/MultiUser/ColorCorrectRegionsMultiUserManager.h b/Engine/Plugins/Experimental/ColorCorrectRegions/Source/ColorCorrectRegions/Private/MultiUser/ColorCorrectRegionsMultiUserManager.h
new file mode 100644
--- /dev/null
+++ b/Engine/Plugins/Experimental/ColorCorrectRegions/Source/ColorCorrectRegions/Private/MultiUser/ColorCorrectRegionsMultiUserManager.h
@@ -0,0 +1,60 @@
+// Copyright Epic Games, Inc. All Rights Reserved.
+
+#pragma once
+
+#if WITH_EDITOR
+
+#include "Delegates/IDelegateInstance.h"
+#include "Templates/SharedPointer.h"
+#include "UObject/WeakObjectPtr.h"
+
+class AColorCorrectRegion;
+class IConcertClient;
+class IConcertClientSession;
+class UColorCorrectRegionsSubsystem;
+struct FConcertCCRPerActorAssignmentEvent;
+struct FConcertSessionContext;
+enum class EConcertConnectionStatus : uint8;
+
+/**
+ * Manages multi-user synchronization of per-actor CC stencil ID assignments.
+ *
+ * When the per-actor CC AffectedActors list changes on one machine, this manager sends a
+ * FConcertCCRPerActorAssignmentEvent Concert custom event. Passive clients in the session
+ * receive the event and re-apply stencil assignments immediately rather than waiting for
+ * the subsystem tick to detect the missing PerAffectedActorStencilData.
+ *
+ * Owned by UColorCorrectRegionsSubsystem; lifetime is aligned with the subsystem.
+ */
+class FCCRMultiUserManager
+{
+public:
+	explicit FCCRMultiUserManager(UColorCorrectRegionsSubsystem* InSubsystem);
+	~FCCRMultiUserManager();
+
+	/**
+	 * Sends a Concert custom event to all connected clients indicating that per-actor CC
+	 * stencil assignment changed for the given region. Should be called after local stencil
+	 * assignment is complete. Does nothing when no Concert session is active.
+	 */
+	void SendPerActorAssignmentEvent(AColorCorrectRegion* Region);
+
+private:
+	void OnSessionStartup(TSharedRef<IConcertClientSession> Session);
+	void OnSessionShutdown(TSharedRef<IConcertClientSession> Session);
+	void OnSessionConnectionChanged(IConcertClientSession& Session, EConcertConnectionStatus ConnectionStatus);
+
+	/** Called when a remote client sends a FConcertCCRPerActorAssignmentEvent. */
+	void HandlePerActorAssignmentEvent(
+		const FConcertSessionContext& Context,
+		const FConcertCCRPerActorAssignmentEvent& Event);
+
+	TWeakObjectPtr<UColorCorrectRegionsSubsystem> WeakSubsystem;
+	TWeakPtr<IConcertClientSession> WeakSession;
+	TWeakPtr<IConcertClient> ConcertClient;
+	FDelegateHandle SessionStartupHandle;
+	FDelegateHandle SessionShutdownHandle;
+	FDelegateHandle SessionConnectionChangedHandle;
+};
+
+#endif // WITH_EDITOR
diff --git a/Engine/Plugins/Experimental/ColorCorrectRegions/Source/ColorCorrectRegions/Private/MultiUser/ColorCorrectRegionsSyncData.h b/Engine/Plugins/Experimental/ColorCorrectRegions/Source/ColorCorrectRegions/Private/MultiUser/ColorCorrectRegionsSyncData.h
new file mode 100644
--- /dev/null
+++ b/Engine/Plugins/Experimental/ColorCorrectRegions/Source/ColorCorrectRegions/Private/MultiUser/ColorCorrectRegionsSyncData.h
@@ -0,0 +1,30 @@
+// Copyright Epic Games, Inc. All Rights Reserved.
+
+#pragma once
+
+#include "CoreMinimal.h"
+#include "UObject/SoftObjectPath.h"
+
+#include "ColorCorrectRegionsSyncData.generated.h"
+
+/**
+ * Concert custom event sent when the per-actor CC AffectedActors list changes on a Color Correction Actor.
+ *
+ * The event is self-contained: it carries the full current AffectedActors state so the receiving
+ * client can apply stencil assignments immediately without waiting for Concert's property transaction
+ * to arrive and update AffectedActors first.
+ */
+USTRUCT()
+struct FConcertCCRPerActorAssignmentEvent
+{
+	GENERATED_BODY()
+
+	/** Soft path to the CCR actor whose per-actor CC assignment changed. */
+	UPROPERTY()
+	FSoftObjectPath CCRActorPath;
+
+	/** Full AffectedActors list as it stands on the sending client after the change. */
+	UPROPERTY()
+	TArray<FSoftObjectPath> AffectedActors;
+};
+
