ColorCorrectActors: Promote CustomDepth/SceneDepth from half to float to preserve precision

#rnx
#Synced-CL 0

diff --git a/Engine/Plugins/Experimental/ColorCorrectRegions/Shaders/Private/ColorCorrectRegionsStencilMerger.usf b/Engine/Plugins/Experimental/ColorCorrectRegions/Shaders/Private/ColorCorrectRegionsStencilMerger.usf
--- a/Engine/Plugins/Experimental/ColorCorrectRegions/Shaders/Private/ColorCorrectRegionsStencilMerger.usf
+++ b/Engine/Plugins/Experimental/ColorCorrectRegions/Shaders/Private/ColorCorrectRegionsStencilMerger.usf
@@ -20,13 +20,13 @@ uint MainPS(
 	float4 SvPosition = UVAndScreenPos;
 	const float Float_Eps = 0.00000000001;
 
-	half CustomDepth = CalcSceneCustomDepth(UVAndScreenPos.xy);
-	half SceneDepth = CalcSceneDepth(UVAndScreenPos.xy);
+	float CustomDepth = CalcSceneCustomDepth(UVAndScreenPos.xy);
+	float SceneDepth = CalcSceneDepth(UVAndScreenPos.xy);
 
-	int CustomStencil = SceneTextureLookup(UVAndScreenPos.xy, PPI_CustomStencil, false).r;
-	for (uint Index = 0; Index < StencilIdCount; Index++)
+	if (abs(SceneDepth - CustomDepth) <= Float_Eps)
 	{
-		if (abs(SceneDepth - CustomDepth) <= Float_Eps)
+		int CustomStencil = SceneTextureLookup(UVAndScreenPos.xy, PPI_CustomStencil, false).r;
+		for (uint Index = 0; Index < StencilIdCount; Index++)
 		{
 			if (CustomStencil == StencilIds[Index])
 			{
